Update description for HAPMedia plugin.

diff --git a/Engine/Plugins/Media/HAPMedia/HAPMedia.uplugin b/Engine/Plugins/Media/HAPMedia/HAPMedia.uplugin
--- a/Engine/Plugins/Media/HAPMedia/HAPMedia.uplugin
+++ b/Engine/Plugins/Media/HAPMedia/HAPMedia.uplugin
@@ -1,12 +1,12 @@
 {
 	"FileVersion": 3,
-	"Version": 2,
-	"VersionName": "2.0",
+	"Version": 1,
+	"VersionName": "1.0",
 	"FriendlyName": "HAP Media",
-	"Description": "Implements HAP codec.",
-	"Category": "Media Players",
+	"Description": "Implements video playback of the HAP Codec. HAP is a high performance, high resolution codec that runs on the GPU.",
 	"CreatedBy": "Epic Games, Inc.",
 	"CreatedByURL": "http://epicgames.com",
+	"Category": "Media Players",
 	"EnabledByDefault": false,
 	"CanContainContent": false,
 	"IsBetaVersion": false,
