Fixed assert/crash when trying to use the HAP video codec with DX12

diff --git a/Engine/Plugins/Media/HAPMedia/Source/HAPMedia/Private/HAPDecoder/WmfMediaHAPDecoder.cpp b/Engine/Plugins/Media/HAPMedia/Source/HAPMedia/Private/HAPDecoder/WmfMediaHAPDecoder.cpp
--- a/Engine/Plugins/Media/HAPMedia/Source/HAPMedia/Private/HAPDecoder/WmfMediaHAPDecoder.cpp
+++ b/Engine/Plugins/Media/HAPMedia/Source/HAPMedia/Private/HAPDecoder/WmfMediaHAPDecoder.cpp
@@ -215,7 +215,7 @@ HRESULT WmfMediaHAPDecoder::ProcessMessage(MFT_MESSAGE_TYPE eMessage, ULONG_PTR
 						if (D3DImmediateContext.IsValid())
 						{
 							UE_LOG(LogHAPMedia, Verbose, TEXT("D3D11Device from Device manager: %p"), D3D11Device.Get());
-							if (InitPipeline())
+							if (!InitPipeline())
 							{
 								UE_LOG(LogHAPMedia, Error, TEXT("Unable to initialize 3D pipeline for Hap Decoding"));
 								hr = E_NOTIMPL;
@@ -256,6 +256,12 @@ HRESULT WmfMediaHAPDecoder::ProcessOutput(DWORD dwFlags, DWORD cOutputBufferCoun
 		return MF_E_TRANSFORM_NEED_MORE_INPUT;
 	}
 
+	if (!D3D11Device.IsValid())
+	{
+		UE_LOG(LogHAPMedia, Error, TEXT("DX11 Device is not initialized!"));
+		return E_NOTIMPL;
+	}
+
 	if (!OutputTexture.IsValid())
 	{
 		TComPtr<ID3D11Texture2D> Texture;
@@ -842,7 +848,7 @@ bool WmfMediaHAPDecoder::InitPipeline()
 		UE_LOG(LogHAPMedia, Error, TEXT("D3DCompile Error/warning: %s"), *FString(pMessage));
 		if (!PSCode.IsValid())
 		{
-			return true;
+			return false;
 		}
 	}
 
@@ -853,7 +859,7 @@ bool WmfMediaHAPDecoder::InitPipeline()
 		UE_LOG(LogHAPMedia, Error, TEXT("D3DCompile Error/warning: %s"), *FString(pMessage));
 		if (!VSCode.IsValid())
 		{
-			return true;
+			return false;
 		}
 	}
 
@@ -863,13 +869,13 @@ bool WmfMediaHAPDecoder::InitPipeline()
 	Result = D3D11Device->CreatePixelShader(PSCode->GetBufferPointer(), PSCode->GetBufferSize(), nullptr, &ps);
 	if (FAILED(Result))
 	{
-		return true;
+		return false;
 	}
 
 	Result = D3D11Device->CreateVertexShader(VSCode->GetBufferPointer(), VSCode->GetBufferSize(), nullptr, &vs);
 	if (FAILED(Result))
 	{
-		return true;
+		return false;
 	}
 
 	D3D11_BUFFER_DESC PixelConstantBufferDesc;
@@ -882,7 +888,7 @@ bool WmfMediaHAPDecoder::InitPipeline()
 	Result = D3D11Device->CreateBuffer(&PixelConstantBufferDesc, nullptr, &PixelBuffer);
 	if (FAILED(Result))
 	{
-		return true;
+		return false;
 	}
 
 	D3DImmediateContext->PSSetConstantBuffers(0, 1, &PixelBuffer);
@@ -902,7 +908,7 @@ bool WmfMediaHAPDecoder::InitPipeline()
 	Result = D3D11Device->CreateInputLayout(InputElementDescription, 2, VSCode->GetBufferPointer(), VSCode->GetBufferSize(), &InputLayout);
 	if (FAILED(Result))
 	{
-		return true;
+		return false;
 	}
 
 	D3DImmediateContext->IASetInputLayout(InputLayout);
@@ -916,7 +922,7 @@ bool WmfMediaHAPDecoder::InitPipeline()
 	Result = D3D11Device->CreateBlendState(&BlendStateDesc, &BlendState);
 	if (FAILED(Result))
 	{
-		return true;
+		return false;
 	}
 
 	float BlendFactor[4] = { 1.0f, 1.0f, 1.0f, 1.0f };
@@ -941,20 +947,20 @@ bool WmfMediaHAPDecoder::InitPipeline()
 	Result = D3D11Device->CreateBuffer(&BufferDesc, NULL, &VertexBuffer);
 	if (FAILED(Result))
 	{
-		return true;
+		return false;
 	}
 
 	D3D11_MAPPED_SUBRESOURCE MappedSubresource;
 	Result = D3DImmediateContext->Map(VertexBuffer, NULL, D3D11_MAP_WRITE_DISCARD, NULL, &MappedSubresource);
 	if (FAILED(Result))
 	{
-		return true;
+		return false;
 	}
 
 	memcpy(MappedSubresource.pData, Vertices, sizeof(Vertices));
 	D3DImmediateContext->Unmap(VertexBuffer, NULL);
 
-	return false;
+	return true;
 }
 
 
