Switching from D3D10 enum to D3D11.

diff --git a/Engine/Plugins/Media/HAPMedia/Source/HAPMedia/Private/HAPDecoder/WmfMediaHAPDecoder.cpp b/Engine/Plugins/Media/HAPMedia/Source/HAPMedia/Private/HAPDecoder/WmfMediaHAPDecoder.cpp
--- a/Engine/Plugins/Media/HAPMedia/Source/HAPMedia/Private/HAPDecoder/WmfMediaHAPDecoder.cpp
+++ b/Engine/Plugins/Media/HAPMedia/Source/HAPMedia/Private/HAPDecoder/WmfMediaHAPDecoder.cpp
@@ -970,7 +970,7 @@ void WmfMediaHAPDecoder::CreateSRV(ID3D11Texture2D* InTexture, D3D11_TEXTURE2D_D
 	memset(&ViewDesc, 0, sizeof(ViewDesc));
 
 	ViewDesc.Format = InTextureDesc.Format;
-	ViewDesc.ViewDimension = D3D10_SRV_DIMENSION_TEXTURE2D;
+	ViewDesc.ViewDimension = D3D11_SRV_DIMENSION_TEXTURE2D;
 	ViewDesc.Texture2D.MipLevels = 1;
 	ViewDesc.Texture2D.MostDetailedMip = 0;
 
