WmfMedia: HAP now uses external buffers.
#rnx

diff --git a/Engine/Plugins/Media/HAPMedia/Source/HAPMedia/Private/HAPDecoder/WmfMediaHAPDecoder.cpp b/Engine/Plugins/Media/HAPMedia/Source/HAPMedia/Private/HAPDecoder/WmfMediaHAPDecoder.cpp
--- a/Engine/Plugins/Media/HAPMedia/Source/HAPMedia/Private/HAPDecoder/WmfMediaHAPDecoder.cpp
+++ b/Engine/Plugins/Media/HAPMedia/Source/HAPMedia/Private/HAPDecoder/WmfMediaHAPDecoder.cpp
@@ -126,13 +126,35 @@ HRESULT WmfMediaHAPDecoder::GetOutputAvailableType(DWORD dwOutputStreamID, DWORD
 
 	if (SUCCEEDED(hr))
 	{
-		if (InputSubType == DecoderGUID_HAP_ALPHA || InputSubType == DecoderGUID_HAP_Q_ALPHA)
+		if (bIsExternalBufferEnabled)
 		{
-			hr = TempOutputType->SetGUID(MF_MT_SUBTYPE, MFVideoFormat_ARGB32);
+			if (InputSubType == DecoderGUID_HAP)
+			{
+				hr = TempOutputType->SetGUID(MF_MT_SUBTYPE, MFVideoFormat_L8);
+			}
+			else if (InputSubType == DecoderGUID_HAP_ALPHA)
+			{
+				hr = TempOutputType->SetGUID(MF_MT_SUBTYPE, MFVideoFormat_L16);
+			}
+			else if (InputSubType == DecoderGUID_HAP_Q)
+			{
+				hr = TempOutputType->SetGUID(MF_MT_SUBTYPE, MFVideoFormat_RGB8);
+			}
+			else
+			{
+				hr = TempOutputType->SetGUID(MF_MT_SUBTYPE, MFVideoFormat_D16);
+			}
 		}
 		else
 		{
-			hr = TempOutputType->SetGUID(MF_MT_SUBTYPE, MFVideoFormat_NV12);
+			if (InputSubType == DecoderGUID_HAP_ALPHA || InputSubType == DecoderGUID_HAP_Q_ALPHA)
+			{
+				hr = TempOutputType->SetGUID(MF_MT_SUBTYPE, MFVideoFormat_ARGB32);
+			}
+			else
+			{
+				hr = TempOutputType->SetGUID(MF_MT_SUBTYPE, MFVideoFormat_NV12);
+			}
 		}
 	}
 
@@ -425,6 +447,12 @@ HRESULT WmfMediaHAPDecoder::ProcessOutput(DWORD dwFlags, DWORD cOutputBufferCoun
 }
 
 
+bool WmfMediaHAPDecoder::IsExternalBufferSupported() const
+{
+	return true;
+}
+
+
 HRESULT WmfMediaHAPDecoder::InternalProcessOutput(IMFSample* InSample)
 {
 	if (OutputQueue.IsEmpty())
@@ -480,89 +508,128 @@ HRESULT WmfMediaHAPDecoder::InternalProcessOutput(IMFSample* InSample)
 	DataBuffer OuputDataBuffer;
 	OutputQueue.Dequeue(OuputDataBuffer);
 
-	D3D11_MAPPED_SUBRESOURCE MappedResourceColor;
-	Result = D3DImmediateContext->Map(WorkTextures->InputTextureColor, 0, D3D11_MAP_WRITE_DISCARD, 0, &MappedResourceColor);
-	memcpy(MappedResourceColor.pData, OuputDataBuffer.Color.GetData(), MappedResourceColor.DepthPitch);
-	D3DImmediateContext->Unmap(WorkTextures->InputTextureColor, 0);
+	LONGLONG TimeStamp = OuputDataBuffer.TimeStamp;
 
-	if (InputSubType == DecoderGUID_HAP_Q_ALPHA)
+	// Are we using external buffers?
+	if (bIsExternalBufferEnabled == false)
 	{
-		D3D11_MAPPED_SUBRESOURCE MappedResourceAlpha;
-		Result = D3DImmediateContext->Map(WorkTextures->InputTextureAlpha, 0, D3D11_MAP_WRITE_DISCARD, 0, &MappedResourceAlpha);
-		memcpy(MappedResourceAlpha.pData, OuputDataBuffer.Alpha.GetData(), MappedResourceAlpha.DepthPitch);
-		D3DImmediateContext->Unmap(WorkTextures->InputTextureAlpha, 0);
-	}
+		D3D11_MAPPED_SUBRESOURCE MappedResourceColor;
+		Result = D3DImmediateContext->Map(WorkTextures->InputTextureColor, 0, D3D11_MAP_WRITE_DISCARD, 0, &MappedResourceColor);
+		memcpy(MappedResourceColor.pData, OuputDataBuffer.Color.GetData(), MappedResourceColor.DepthPitch);
+		D3DImmediateContext->Unmap(WorkTextures->InputTextureColor, 0);
 
+		if (InputSubType == DecoderGUID_HAP_Q_ALPHA)
+		{
+			D3D11_MAPPED_SUBRESOURCE MappedResourceAlpha;
+			Result = D3DImmediateContext->Map(WorkTextures->InputTextureAlpha, 0, D3D11_MAP_WRITE_DISCARD, 0, &MappedResourceAlpha);
+			memcpy(MappedResourceAlpha.pData, OuputDataBuffer.Alpha.GetData(), MappedResourceAlpha.DepthPitch);
+			D3DImmediateContext->Unmap(WorkTextures->InputTextureAlpha, 0);
+		}
 
-	if (InputSubType == DecoderGUID_HAP_Q_ALPHA)
-	{
-		BindSRV(WorkTextures->SRV_InputColor, WorkTextures->SRV_InputAlpha);
-	}
-	else
-	{
-		BindSRV(WorkTextures->SRV_InputColor);
-	}
 
-	UINT stride = sizeof(VertexDescription);
-	UINT offset = 0;
+		if (InputSubType == DecoderGUID_HAP_Q_ALPHA)
+		{
+			BindSRV(WorkTextures->SRV_InputColor, WorkTextures->SRV_InputAlpha);
+		}
+		else
+		{
+			BindSRV(WorkTextures->SRV_InputColor);
+		}
 
-	if (InputSubType == DecoderGUID_HAP)
-	{
-		FrameData.Mode = 0;
-	}
-	else if (InputSubType == DecoderGUID_HAP_ALPHA)
-	{
-		FrameData.Mode = 1;
-	}
-	else if (InputSubType == DecoderGUID_HAP_Q)
-	{
-		FrameData.Mode = 2;
-	}
-	else // if (InputSubType == DecoderGUID_HAP_Q_ALPHA)
-	{
-		FrameData.Mode = 3;
-	}
+		UINT stride = sizeof(VertexDescription);
+		UINT offset = 0;
 
-	D3D11_VIEWPORT viewport;
-	ZeroMemory(&viewport, sizeof(D3D11_VIEWPORT));
-	viewport.TopLeftX = 0;
-	viewport.TopLeftY = 0;
+		if (InputSubType == DecoderGUID_HAP)
+		{
+			FrameData.Mode = 0;
+		}
+		else if (InputSubType == DecoderGUID_HAP_ALPHA)
+		{
+			FrameData.Mode = 1;
+		}
+		else if (InputSubType == DecoderGUID_HAP_Q)
+		{
+			FrameData.Mode = 2;
+		}
+		else // if (InputSubType == DecoderGUID_HAP_Q_ALPHA)
+		{
+			FrameData.Mode = 3;
+		}
 
-	ID3D11RenderTargetView* m_pRenderViews[1];
+		D3D11_VIEWPORT viewport;
+		ZeroMemory(&viewport, sizeof(D3D11_VIEWPORT));
+		viewport.TopLeftX = 0;
+		viewport.TopLeftY = 0;
 
-	int StartPass = 0;
+		ID3D11RenderTargetView* m_pRenderViews[1];
 
-	if (InputSubType == DecoderGUID_HAP_ALPHA || InputSubType == DecoderGUID_HAP_Q_ALPHA)
-	{
-		StartPass = 1;
-	}
+		int StartPass = 0;
+
+		if (InputSubType == DecoderGUID_HAP_ALPHA || InputSubType == DecoderGUID_HAP_Q_ALPHA)
+		{
+			StartPass = 1;
+		}
+
+		for (int PassIndex = StartPass; PassIndex < 2; PassIndex++)
+		{
+			if (PassIndex == 0)
+			{
+				m_pRenderViews[0] = OutputTexture->RTV_B;
+				FrameData.WriteY = 0;
+				viewport.Width = float(ImageWidthInPixels / 2);
+				viewport.Height = float(ImageHeightInPixels / 2);
+			}
+			else
+			{
+				m_pRenderViews[0] = OutputTexture->RTV_A;
+				FrameData.WriteY = 1;
+				viewport.Width = float(ImageWidthInPixels);
+				viewport.Height = float(ImageHeightInPixels);
+			}
 
-	for (int PassIndex = StartPass; PassIndex < 2; PassIndex++)
+			D3DImmediateContext->RSSetViewports(1, &viewport);
+			D3DImmediateContext->UpdateSubresource(PixelBuffer, 0, nullptr, &FrameData, 0, 0);
+			D3DImmediateContext->OMSetRenderTargets(1, &m_pRenderViews[0], NULL);
+			D3DImmediateContext->IASetVertexBuffers(0, 1, &VertexBuffer, &stride, &offset);
+			D3DImmediateContext->IASetPrimitiveTopology(D3D11_PRIMITIVE_TOPOLOGY_TRIANGLESTRIP);
+			D3DImmediateContext->Draw(4, 0);
+		}
+	}
+	else
 	{
-		if (PassIndex == 0)
+		// Get buffer.
+		int32 BufferSize = ImageWidthInPixels * ImageHeightInPixels;
+		int32 AlphaBufferSize = 0;
+
+		// HapQ Alpha has a separate alpha buffer.
+		if (InputSubType == DecoderGUID_HAP_Q_ALPHA)
 		{
-			m_pRenderViews[0] = OutputTexture->RTV_B;
-			FrameData.WriteY = 0;
-			viewport.Width = float(ImageWidthInPixels / 2);
-			viewport.Height = float(ImageHeightInPixels / 2);
+			AlphaBufferSize = BufferSize / 2;
 		}
-		else
+
+		// Hap is 4 bits per pixel.
+		if (InputSubType == DecoderGUID_HAP)
 		{
-			m_pRenderViews[0] = OutputTexture->RTV_A;
-			FrameData.WriteY = 1;
-			viewport.Width = float(ImageWidthInPixels);
-			viewport.Height = float(ImageHeightInPixels);
+			BufferSize = BufferSize / 2;
 		}
 
-		D3DImmediateContext->RSSetViewports(1, &viewport);
-		D3DImmediateContext->UpdateSubresource(PixelBuffer, 0, nullptr, &FrameData, 0, 0);
-		D3DImmediateContext->OMSetRenderTargets(1, &m_pRenderViews[0], NULL);
-		D3DImmediateContext->IASetVertexBuffers(0, 1, &VertexBuffer, &stride, &offset);
-		D3DImmediateContext->IASetPrimitiveTopology(D3D11_PRIMITIVE_TOPOLOGY_TRIANGLESTRIP);
-		D3DImmediateContext->Draw(4, 0);
-	}
+		// Get external buffer.
+		TArray<uint8>* pExternalBuffer = AllocateExternalBuffer(TimeStamp, BufferSize + AlphaBufferSize);
+		if (pExternalBuffer == nullptr)
+		{
+			return 0;
+		}
 
-	LONGLONG TimeStamp = OuputDataBuffer.TimeStamp;
+		// Copy to buffer.
+		pExternalBuffer->SetNum(BufferSize + AlphaBufferSize);
+		FMemory::Memcpy(pExternalBuffer->GetData(), OuputDataBuffer.Color.GetData(), BufferSize);
+
+		if (AlphaBufferSize > 0)
+		{
+			FMemory::Memcpy(pExternalBuffer->GetData() + BufferSize, OuputDataBuffer.Alpha.GetData(), AlphaBufferSize);
+		}
+	}
+	
 
 	InputQueue.Enqueue(MoveTemp(OuputDataBuffer));
 
@@ -672,7 +739,14 @@ HRESULT WmfMediaHAPDecoder::OnSetInputType(IMFMediaType* InMediaType)
 	if (SUCCEEDED(hr))
 	{
 		InputImageSize = ImageWidthInPixels * ImageHeightInPixels * 4;
-		OutputImageSize = ImageWidthInPixels * ImageHeightInPixels * 3 / 2;
+		if (bIsExternalBufferEnabled)
+		{
+			OutputImageSize = ImageWidthInPixels * ImageHeightInPixels / 2;
+		}
+		else
+		{
+			OutputImageSize = ImageWidthInPixels * ImageHeightInPixels * 3 / 2;
+		}
 
 		InputType = InMediaType;
 		InputType->AddRef();
diff --git a/Engine/Plugins/Media/HAPMedia/Source/HAPMedia/Private/HAPDecoder/WmfMediaHAPDecoder.h b/Engine/Plugins/Media/HAPMedia/Source/HAPMedia/Private/HAPDecoder/WmfMediaHAPDecoder.h
--- a/Engine/Plugins/Media/HAPMedia/Source/HAPMedia/Private/HAPDecoder/WmfMediaHAPDecoder.h
+++ b/Engine/Plugins/Media/HAPMedia/Source/HAPMedia/Private/HAPDecoder/WmfMediaHAPDecoder.h
@@ -29,6 +29,8 @@ public:
 	virtual HRESULT STDMETHODCALLTYPE ProcessMessage(MFT_MESSAGE_TYPE eMessage, ULONG_PTR ulParam) override;
 	virtual HRESULT STDMETHODCALLTYPE ProcessOutput(DWORD dwFlags, DWORD cOutputBufferCount, MFT_OUTPUT_DATA_BUFFER* pOutputSamples, DWORD* pdwStatus) override;
 
+	virtual bool IsExternalBufferSupported() const override;
+
 	static bool IsSupported(const GUID& InGuid);
 	static bool SetOutputFormat(const GUID& InGuid, GUID& OutVideoFormat);
 
diff --git a/Engine/Plugins/Media/WmfMedia/Shaders/Private/MediaHardwareVideoDecoding.usf b/Engine/Plugins/Media/WmfMedia/Shaders/Private/MediaHardwareVideoDecoding.usf
--- a/Engine/Plugins/Media/WmfMedia/Shaders/Private/MediaHardwareVideoDecoding.usf
+++ b/Engine/Plugins/Media/WmfMedia/Shaders/Private/MediaHardwareVideoDecoding.usf
@@ -27,6 +27,8 @@ static const float2 kQuadUVs[6] =
 	float2(0.0, 0.0),
 };
 
+static const float4 kYCoCgOffsets = float4(-0.50196078431373, -0.50196078431373, 0, 0);
+
 Texture2D TextureY;
 Texture2D TextureUV;
 SamplerState PointClampedSamplerY;
@@ -74,3 +76,41 @@ void Y416ConvertPS(in noperspective float2 InUV : TEXCOORD0
 	OutColor.rgb = YuvToRgb(AVYU.yzw, ColorTransform, SrgbToLinear);
 	OutColor.a = AVYU.x;
 }
+
+void YCoCgConvertPS(in noperspective float2 InUV : TEXCOORD0
+		  , in float4 SvPosition : SV_POSITION
+		  , out float4 OutColor : SV_Target0)
+{
+	float4 CoCgSY = TextureY.Sample(BilinearClampedSamplerUV, InUV);
+
+	CoCgSY += kYCoCgOffsets;
+	
+	float Scale = (CoCgSY.z * (255.0 / 8.0)) + 1.0;
+	
+	float Co = CoCgSY.x / Scale;
+	float Cg = CoCgSY.y / Scale;
+	float Y = CoCgSY.w;
+	
+	OutColor = float4(Y + Co - Cg, Y + Cg, Y - Co - Cg, 1.0);
+	OutColor.xyz = sRGBToLinear(OutColor.xyz);
+}
+
+void YCoCgAlphaConvertPS(in noperspective float2 InUV : TEXCOORD0
+		  , in float4 SvPosition : SV_POSITION
+		  , out float4 OutColor : SV_Target0)
+{
+	float4 CoCgSY = TextureY.Sample(BilinearClampedSamplerUV, InUV);
+	float Alpha = TextureUV.Sample(BilinearClampedSamplerUV, InUV);
+
+	CoCgSY += kYCoCgOffsets;
+	
+	float Scale = (CoCgSY.z * (255.0 / 8.0)) + 1.0;
+	
+	float Co = CoCgSY.x / Scale;
+	float Cg = CoCgSY.y / Scale;
+	float Y = CoCgSY.w;
+	
+	OutColor = float4(Y + Co - Cg, Y + Cg, Y - Co - Cg, Alpha);
+	OutColor.xyz = sRGBToLinear(OutColor.xyz);
+}
+
diff --git a/Engine/Plugins/Media/WmfMedia/Source/WmfMedia/Private/Player/WmfMediaTextureSample.h b/Engine/Plugins/Media/WmfMedia/Source/WmfMedia/Private/Player/WmfMediaTextureSample.h
--- a/Engine/Plugins/Media/WmfMedia/Source/WmfMedia/Private/Player/WmfMediaTextureSample.h
+++ b/Engine/Plugins/Media/WmfMedia/Source/WmfMedia/Private/Player/WmfMediaTextureSample.h
@@ -115,11 +115,6 @@ public:
 			return false;
 		}
 
-		if ((InStride * InDim.Y) > InSize)
-		{
-			return false;
-		}
-
 		Buffer = MoveTemp(*InBuffer);
 
 		Duration = InDuration;
@@ -132,6 +127,13 @@ public:
 		return true;
 	}
 
+	/**
+	 * See if we are using external buffers.
+	 * 
+	 * @return True if so.
+	 */
+	bool IsBufferExternal() const { return bIsBufferExternal; }
+
 public:
 
 	//~ IMediaTextureSample interface
diff --git a/Engine/Plugins/Media/WmfMedia/Source/WmfMedia/Private/Wmf/WmfMediaHardwareVideoDecodingRendering.cpp b/Engine/Plugins/Media/WmfMedia/Source/WmfMedia/Private/Wmf/WmfMediaHardwareVideoDecodingRendering.cpp
--- a/Engine/Plugins/Media/WmfMedia/Source/WmfMedia/Private/Wmf/WmfMediaHardwareVideoDecodingRendering.cpp
+++ b/Engine/Plugins/Media/WmfMedia/Source/WmfMedia/Private/Wmf/WmfMediaHardwareVideoDecodingRendering.cpp
@@ -120,12 +120,19 @@ bool FWmfMediaHardwareVideoDecodingParameters::ConvertTextureFormat_RenderThread
 
 	TComPtr<ID3D11Texture2D> SampleTexture = InSample->GetSourceTexture();
 
-	ID3D11Device* D3D11Device = static_cast<ID3D11Device*>(GDynamicRHI->RHIGetNativeDevice());
+	ID3D11Device* D3D11Device = nullptr;
+	if (InSample->IsBufferExternal() == false)
+	{
+		D3D11Device = static_cast<ID3D11Device*>(GDynamicRHI->RHIGetNativeDevice());
+	}
 	ID3D11DeviceContext* D3D11DeviceContext = nullptr;
 			
 	// Must access rendering device context to copy shared resource.
-	D3D11Device->GetImmediateContext(&D3D11DeviceContext);
-	if (D3D11DeviceContext)
+	if (D3D11Device != nullptr)
+	{
+		D3D11Device->GetImmediateContext(&D3D11DeviceContext);
+	}
+	
 	{
 		FRHICommandListImmediate& RHICmdList = FRHICommandListExecutor::GetImmediateCommandList();
 
@@ -147,16 +154,28 @@ bool FWmfMediaHardwareVideoDecodingParameters::ConvertTextureFormat_RenderThread
 		GraphicsPSOInit.PrimitiveType = PT_TriangleList;
 		GraphicsPSOInit.BoundShaderState.VertexDeclarationRHI = GetVertexDeclarationFVector4();
 
+		// Copy buffer.
 		FTexture2DRHIRef SampleDestinationTexture = InSample->GetOrCreateDestinationTexture();
-
-		if (RHICmdList.Bypass())
+		if (D3D11Device == nullptr)
 		{
-			FRHICommandCopyResource Cmd(SampleTexture, SampleDestinationTexture);
-			Cmd.Execute(RHICmdList);
+			const uint8* Data = (const uint8*)InSample->GetBuffer();
+			
+			uint32 Stride = InSample->GetStride();
+			uint32 Height = InSample->GetDim().Y;
+			FUpdateTextureRegion2D Region(0, 0, 0, 0, InSample->GetDim().X, Height);
+			RHIUpdateTexture2D(SampleDestinationTexture, 0, Region, Stride, Data);
 		}
 		else
 		{
-			new (RHICmdList.AllocCommand<FRHICommandCopyResource>()) FRHICommandCopyResource(SampleTexture, SampleDestinationTexture);
+			if (RHICmdList.Bypass())
+			{
+				FRHICommandCopyResource Cmd(SampleTexture, SampleDestinationTexture);
+				Cmd.Execute(RHICmdList);
+			}
+			else
+			{
+				new (RHICmdList.AllocCommand<FRHICommandCopyResource>()) FRHICommandCopyResource(SampleTexture, SampleDestinationTexture);
+			}
 		}
 
 		FGlobalShaderMap* GlobalShaderMap = GetGlobalShaderMap(GMaxRHIFeatureLevel);
@@ -176,7 +195,9 @@ bool FWmfMediaHardwareVideoDecodingParameters::ConvertTextureFormat_RenderThread
 			VertexShader->SetParameters(RHICmdList, VertexShader.GetVertexShader(), Y_SRV, UV_SRV, InSample->IsOutputSrgb());
 			PixelShader->SetParameters(RHICmdList, PixelShader.GetPixelShader(), Y_SRV, UV_SRV, InSample->IsOutputSrgb());
 		}
-		else if (InSample->GetFormat() == EMediaTextureSampleFormat::CharBGRA)
+		else if ((InSample->GetFormat() == EMediaTextureSampleFormat::CharBGRA) ||
+			(InSample->GetFormat() == EMediaTextureSampleFormat::DXT1) ||
+			(InSample->GetFormat() == EMediaTextureSampleFormat::DXT5))
 		{
 			TShaderMapRef< FHardwareVideoDecodingVS > VertexShader(GlobalShaderMap);
 			TShaderMapRef< FHardwareVideoDecodingPassThroughPS > PixelShader(GlobalShaderMap);
@@ -185,7 +206,16 @@ bool FWmfMediaHardwareVideoDecodingParameters::ConvertTextureFormat_RenderThread
 
 			SetGraphicsPipelineState(RHICmdList, GraphicsPSOInit, 0);
 
-			FShaderResourceViewRHIRef SRV = RHICreateShaderResourceView(SampleDestinationTexture, 0, 1, PF_B8G8R8A8);
+			EPixelFormat PixelFormat = PF_B8G8R8A8;
+			if (InSample->GetFormat() == EMediaTextureSampleFormat::DXT1)
+			{
+				PixelFormat = PF_DXT1;
+			}
+			else if (InSample->GetFormat() == EMediaTextureSampleFormat::DXT5)
+			{
+				PixelFormat = PF_DXT5;
+			}
+			FShaderResourceViewRHIRef SRV = RHICreateShaderResourceView(SampleDestinationTexture, 0, 1, PixelFormat);
 			VertexShader->SetParameters(RHICmdList, VertexShader.GetVertexShader(), SRV, InSample->IsOutputSrgb());
 			PixelShader->SetParameters(RHICmdList, PixelShader.GetPixelShader(), SRV, InSample->IsOutputSrgb());
 		}
@@ -202,11 +232,54 @@ bool FWmfMediaHardwareVideoDecodingParameters::ConvertTextureFormat_RenderThread
 			VertexShader->SetParameters(RHICmdList, VertexShader.GetVertexShader(), SRV, InSample->IsOutputSrgb());
 			PixelShader->SetParameters(RHICmdList, PixelShader.GetPixelShader(), SRV, InSample->IsOutputSrgb());
 		}
+		else if (InSample->GetFormat() == EMediaTextureSampleFormat::YCoCg_DXT5)
+		{
+			TShaderMapRef< FHardwareVideoDecodingVS > VertexShader(GlobalShaderMap);
+			TShaderMapRef< FHardwareVideoDecodingYCoCgPS > PixelShader(GlobalShaderMap);
+			GraphicsPSOInit.BoundShaderState.VertexShaderRHI = VertexShader.GetVertexShader();
+			GraphicsPSOInit.BoundShaderState.PixelShaderRHI = PixelShader.GetPixelShader();
+
+			SetGraphicsPipelineState(RHICmdList, GraphicsPSOInit, 0);
+
+			FShaderResourceViewRHIRef SRV = RHICreateShaderResourceView(SampleDestinationTexture, 0, 1, PF_DXT5);
+			VertexShader->SetParameters(RHICmdList, VertexShader.GetVertexShader(), SRV, InSample->IsOutputSrgb());
+			PixelShader->SetParameters(RHICmdList, PixelShader.GetPixelShader(), SRV, InSample->IsOutputSrgb());
+
+		}
+		else if (InSample->GetFormat() == EMediaTextureSampleFormat::YCoCg_DXT5_Alpha_BC4)
+		{
+			// Set up alpha texture.
+			FTexture2DRHIRef SampleDestinationAlphaTexture = InSample->GetOrCreateDestinationAlphaTexture();
+			{
+				const uint8* Data = ((const uint8*)InSample->GetBuffer()) + InSample->GetDim().X * InSample->GetDim().Y;
+
+				uint32 Stride = InSample->GetDim().X * 2;
+				uint32 Height = InSample->GetDim().Y;
+				FUpdateTextureRegion2D Region(0, 0, 0, 0, InSample->GetDim().X, Height);
+				RHIUpdateTexture2D(SampleDestinationAlphaTexture, 0, Region, Stride, Data);
+			}
+			
+			TShaderMapRef< FHardwareVideoDecodingVS > VertexShader(GlobalShaderMap);
+			TShaderMapRef< FHardwareVideoDecodingYCoCgAlphaPS > PixelShader(GlobalShaderMap);
+			GraphicsPSOInit.BoundShaderState.VertexShaderRHI = VertexShader.GetVertexShader();
+			GraphicsPSOInit.BoundShaderState.PixelShaderRHI = PixelShader.GetPixelShader();
+
+			SetGraphicsPipelineState(RHICmdList, GraphicsPSOInit, 0);
+
+			FShaderResourceViewRHIRef SRV = RHICreateShaderResourceView(SampleDestinationTexture, 0, 1, PF_DXT5);
+			FShaderResourceViewRHIRef SRVAlpha = RHICreateShaderResourceView(SampleDestinationAlphaTexture, 0, 1, PF_BC4);
+			VertexShader->SetParameters(RHICmdList, VertexShader.GetVertexShader(), SRV, InSample->IsOutputSrgb());
+			PixelShader->SetParameters(RHICmdList, PixelShader.GetPixelShader(), SRV, SRVAlpha, InSample->IsOutputSrgb());
+
+		}
 
 		RHICmdList.DrawPrimitive(0, 2, 1);
 		RHICmdList.EndRenderPass();
 
-		D3D11DeviceContext->Release();
+		if (D3D11DeviceContext != nullptr)
+		{
+			D3D11DeviceContext->Release();
+		}
 	}
 
 	return true;
diff --git a/Engine/Plugins/Media/WmfMedia/Source/WmfMedia/Private/Wmf/WmfMediaHardwareVideoDecodingShaders.cpp b/Engine/Plugins/Media/WmfMedia/Source/WmfMedia/Private/Wmf/WmfMediaHardwareVideoDecodingShaders.cpp
--- a/Engine/Plugins/Media/WmfMedia/Source/WmfMedia/Private/Wmf/WmfMediaHardwareVideoDecodingShaders.cpp
+++ b/Engine/Plugins/Media/WmfMedia/Source/WmfMedia/Private/Wmf/WmfMediaHardwareVideoDecodingShaders.cpp
@@ -8,3 +8,5 @@ IMPLEMENT_SHADER_TYPE(, FHardwareVideoDecodingVS, TEXT("/Plugin/WmfMedia/Private
 IMPLEMENT_SHADER_TYPE(, FHardwareVideoDecodingPS, TEXT("/Plugin/WmfMedia/Private/MediaHardwareVideoDecoding.usf"), TEXT("NV12ConvertPS"), SF_Pixel)
 IMPLEMENT_SHADER_TYPE(, FHardwareVideoDecodingPassThroughPS, TEXT("/Plugin/WmfMedia/Private/MediaHardwareVideoDecoding.usf"), TEXT("PassThroughPS"), SF_Pixel)
 IMPLEMENT_SHADER_TYPE(, FHardwareVideoDecodingY416PS, TEXT("/Plugin/WmfMedia/Private/MediaHardwareVideoDecoding.usf"), TEXT("Y416ConvertPS"), SF_Pixel)
+IMPLEMENT_SHADER_TYPE(, FHardwareVideoDecodingYCoCgPS, TEXT("/Plugin/WmfMedia/Private/MediaHardwareVideoDecoding.usf"), TEXT("YCoCgConvertPS"), SF_Pixel)
+IMPLEMENT_SHADER_TYPE(, FHardwareVideoDecodingYCoCgAlphaPS, TEXT("/Plugin/WmfMedia/Private/MediaHardwareVideoDecoding.usf"), TEXT("YCoCgAlphaConvertPS"), SF_Pixel)
diff --git a/Engine/Plugins/Media/WmfMedia/Source/WmfMedia/Private/Wmf/WmfMediaHardwareVideoDecodingTextureSample.cpp b/Engine/Plugins/Media/WmfMedia/Source/WmfMedia/Private/Wmf/WmfMediaHardwareVideoDecodingTextureSample.cpp
--- a/Engine/Plugins/Media/WmfMedia/Source/WmfMedia/Private/Wmf/WmfMediaHardwareVideoDecodingTextureSample.cpp
+++ b/Engine/Plugins/Media/WmfMedia/Source/WmfMedia/Private/Wmf/WmfMediaHardwareVideoDecodingTextureSample.cpp
@@ -92,6 +92,9 @@ void FWmfMediaHardwareVideoDecodingTextureSample::ShutdownPoolable()
 			}
 		}
 	}
+
+	AlphaFormat = PF_Unknown;
+	bIsDestinationTextureSRGB = false;
 }
 #endif
 
diff --git a/Engine/Plugins/Media/WmfMedia/Source/WmfMedia/Private/Wmf/WmfMediaStreamSink.cpp b/Engine/Plugins/Media/WmfMedia/Source/WmfMedia/Private/Wmf/WmfMediaStreamSink.cpp
--- a/Engine/Plugins/Media/WmfMedia/Source/WmfMedia/Private/Wmf/WmfMediaStreamSink.cpp
+++ b/Engine/Plugins/Media/WmfMedia/Source/WmfMedia/Private/Wmf/WmfMediaStreamSink.cpp
@@ -891,7 +891,9 @@ void FWmfMediaStreamSink::CopyTextureAndEnqueueSample(IMFSample* pSample)
 		const TSharedRef<FWmfMediaHardwareVideoDecodingTextureSample, ESPMode::ThreadSafe> TextureSample = VideoSamplePool->AcquireShared();
 
 		EPixelFormat PixelFormat = PF_Unknown;
+		EPixelFormat AlphaPixelFormat = PF_Unknown;
 		EMediaTextureSampleFormat MediaTextureSampleFormat = EMediaTextureSampleFormat::Undefined;
+		bool bIsDestinationTextureSRGB = false;
 
 		if (Guid == MFVideoFormat_NV12)
 		{
@@ -903,6 +905,33 @@ void FWmfMediaStreamSink::CopyTextureAndEnqueueSample(IMFSample* pSample)
 			PixelFormat = PF_B8G8R8A8;
 			MediaTextureSampleFormat = EMediaTextureSampleFormat::CharBGRA;
 		}
+		else if (Guid == MFVideoFormat_L8)
+		{
+			// We misuse this for DXT1 as there is no GUID for it.
+			MediaTextureSampleFormat = EMediaTextureSampleFormat::DXT1;
+			PixelFormat = PF_DXT1;
+			bIsDestinationTextureSRGB = true;
+		}
+		else if (Guid == MFVideoFormat_L16)
+		{
+			// We misuse this for DXT5 as there is no GUID for it.
+			MediaTextureSampleFormat = EMediaTextureSampleFormat::DXT5;
+			PixelFormat = PF_DXT5;
+			bIsDestinationTextureSRGB = true;
+		}
+		else if (Guid == MFVideoFormat_RGB8)
+		{
+			// We misuse this for YCoCg_DXT5 as there is no GUID for it.
+			MediaTextureSampleFormat = EMediaTextureSampleFormat::YCoCg_DXT5;
+			PixelFormat = PF_DXT5;
+		}
+		else if (Guid == MFVideoFormat_D16)
+		{
+			// We misuse this for YCoCg_DXT5_Alpha_BC4 as there is no GUID for it.
+			MediaTextureSampleFormat = EMediaTextureSampleFormat::YCoCg_DXT5_Alpha_BC4;
+			PixelFormat = PF_DXT5;
+			AlphaPixelFormat = PF_BC4;
+		}
 		else // if (Guid == MFVideoFormat_Y416)
 		{
 			PixelFormat = PF_A16B16G16R16;
@@ -921,7 +950,22 @@ void FWmfMediaStreamSink::CopyTextureAndEnqueueSample(IMFSample* pSample)
 				return;
 			}
 
-			LONG Pitch = ExternalBuffer.Num() / DimY;
+			// Set pitch.
+			LONG Pitch = 0;
+			if (MediaTextureSampleFormat == EMediaTextureSampleFormat::DXT1)
+			{
+				Pitch = DimX * 2;
+			}
+			else if ((MediaTextureSampleFormat == EMediaTextureSampleFormat::DXT5) ||
+				(MediaTextureSampleFormat == EMediaTextureSampleFormat::YCoCg_DXT5) ||
+				(MediaTextureSampleFormat == EMediaTextureSampleFormat::YCoCg_DXT5_Alpha_BC4))
+			{
+				Pitch = DimX * 4;
+			}
+			else
+			{
+				Pitch = ExternalBuffer.Num() / DimY;
+			}
 
 			// Set up sample.
 			TextureSample->InitializeExternal(&ExternalBuffer,
@@ -929,6 +973,10 @@ void FWmfMediaStreamSink::CopyTextureAndEnqueueSample(IMFSample* pSample)
 				Pitch,
 				FTimespan::FromMicroseconds(SampleTime / 10),
 				FTimespan::FromMicroseconds(SampleDuration / 10));
+
+			TextureSample->SetPixelFormat(PixelFormat);
+			TextureSample->SetIsDestinationTextureSRGB(bIsDestinationTextureSRGB);
+			TextureSample->SetAlphaTexture(AlphaPixelFormat);
 			VideoSampleQueue->Enqueue(TextureSample);
 			UE_LOG(LogWmfMedia, VeryVerbose, TEXT("Enqueued external buffer onto VideoSampleQueue."));
 		}
diff --git a/Engine/Plugins/Media/WmfMedia/Source/WmfMedia/Private/Wmf/WmfMediaUtils.cpp b/Engine/Plugins/Media/WmfMedia/Source/WmfMedia/Private/Wmf/WmfMediaUtils.cpp
--- a/Engine/Plugins/Media/WmfMedia/Source/WmfMedia/Private/Wmf/WmfMediaUtils.cpp
+++ b/Engine/Plugins/Media/WmfMedia/Source/WmfMedia/Private/Wmf/WmfMediaUtils.cpp
@@ -102,7 +102,7 @@ namespace WmfMedia
 //		&MFVideoFormat_RGB24,
 //		&MFVideoFormat_RGB555,
 //		&MFVideoFormat_RGB565,
-//		&MFVideoFormat_RGB8,
+		&MFVideoFormat_RGB8,
 
 		// 8-bit YUV (packed)
 //		&MFVideoFormat_AYUV,
@@ -134,6 +134,11 @@ namespace WmfMedia
 		// 16-bit YUV (planar)
 //		&MFVideoFormat_P016,
 //		&MFVideoFormat_P216,
+
+		&MFVideoFormat_L8,
+		&MFVideoFormat_L16,
+		&MFVideoFormat_D16,
+
 	};
 
 
diff --git a/Engine/Plugins/Media/WmfMedia/Source/WmfMedia/Public/WmfMediaHardwareVideoDecodingShaders.h b/Engine/Plugins/Media/WmfMedia/Source/WmfMedia/Public/WmfMediaHardwareVideoDecodingShaders.h
--- a/Engine/Plugins/Media/WmfMedia/Source/WmfMedia/Public/WmfMediaHardwareVideoDecodingShaders.h
+++ b/Engine/Plugins/Media/WmfMedia/Source/WmfMedia/Public/WmfMediaHardwareVideoDecodingShaders.h
@@ -149,3 +149,37 @@ public:
 		: FWmfMediaHardwareVideoDecodingShader(Initializer)
 	{ }
 };
+
+
+class FHardwareVideoDecodingYCoCgPS : public FWmfMediaHardwareVideoDecodingShader
+{
+	DECLARE_SHADER_TYPE(FHardwareVideoDecodingYCoCgPS, Global);
+
+public:
+
+	/** Default constructor. */
+	FHardwareVideoDecodingYCoCgPS() {}
+
+	/** Initialization constructor. */
+	FHardwareVideoDecodingYCoCgPS(const ShaderMetaType::CompiledShaderInitializerType& Initializer)
+		: FWmfMediaHardwareVideoDecodingShader(Initializer)
+	{ }
+};
+
+
+class FHardwareVideoDecodingYCoCgAlphaPS : public FWmfMediaHardwareVideoDecodingShader
+{
+	DECLARE_SHADER_TYPE(FHardwareVideoDecodingYCoCgAlphaPS, Global);
+
+public:
+
+	/** Default constructor. */
+	FHardwareVideoDecodingYCoCgAlphaPS() {}
+
+	/** Initialization constructor. */
+	FHardwareVideoDecodingYCoCgAlphaPS(const ShaderMetaType::CompiledShaderInitializerType& Initializer)
+		: FWmfMediaHardwareVideoDecodingShader(Initializer)
+	{ }
+};
+
+
diff --git a/Engine/Plugins/Media/WmfMedia/Source/WmfMedia/Public/WmfMediaHardwareVideoDecodingTextureSample.h b/Engine/Plugins/Media/WmfMedia/Source/WmfMedia/Public/WmfMediaHardwareVideoDecodingTextureSample.h
--- a/Engine/Plugins/Media/WmfMedia/Source/WmfMedia/Public/WmfMediaHardwareVideoDecodingTextureSample.h
+++ b/Engine/Plugins/Media/WmfMedia/Source/WmfMedia/Public/WmfMediaHardwareVideoDecodingTextureSample.h
@@ -29,8 +29,10 @@ public:
 
 	/** Default constructor. */
 	FWmfMediaHardwareVideoDecodingTextureSample()
-		: FWmfMediaTextureSample(),
-		Format(PF_Unknown)
+		: FWmfMediaTextureSample()
+		, Format(PF_Unknown)
+		, AlphaFormat(PF_Unknown)
+		, bIsDestinationTextureSRGB(false)
 	{ }
 
 public:
@@ -49,6 +51,30 @@ public:
 	 */
 	ID3D11Texture2D* InitializeSourceTexture(const TRefCountPtr<ID3D11Device>& InD3D11Device, FTimespan InTime, FTimespan InDuration, const FIntPoint& InDim, EPixelFormat InFormat, EMediaTextureSampleFormat InMediaTextureSampleFormat);
 
+	/**
+	 * Call this to set the pixel format.
+	 * 
+	 * InitializeSourceTexture does this so you only need to call this function if
+	 * you don't call InitializeSourceTexture.
+	 * 
+	 * @param InFormat Pixel format for texture.
+	 */
+	void SetPixelFormat(EPixelFormat InFormat) { Format = InFormat; }
+
+	/**
+	 * Call this to set whether DestinationTexture should be sRGB or not.
+	 * 
+	 * @param bIsSRGB True if so.
+	 */
+	void SetIsDestinationTextureSRGB(bool bIsSRGB) { bIsDestinationTextureSRGB = bIsSRGB; }
+
+	/**
+	 * Call this to set what the alpha texture should be.
+	 * 
+	 * @param InFormat Pixel format of alpha texture.
+	 */
+	void SetAlphaTexture(EPixelFormat InFormat) { AlphaFormat = InFormat; }
+
 	/**
 	 * Get media texture sample converter if sample implements it
 	 *
@@ -59,12 +85,7 @@ public:
 		// Only use sample converter for Win8+
 		if (FPlatformMisc::VerifyWindowsVersion(6, 2))
 		{
-			// Are we using an external buffer?
-			// D3D12 will use a CPU buffer and has no use for this convertor.
-			if (bIsBufferExternal == false)
-			{
-				return this;
-			}
+			return this;
 		}
 		return nullptr;
 	}
@@ -97,11 +118,17 @@ public:
 	{
 		if (DestinationTexture.IsValid() && DestinationTexture->GetSizeX() == Dim.X && DestinationTexture->GetSizeY() == Dim.Y)
 		{
-			return DestinationTexture;
+			ETextureCreateFlags CurrentFlags = DestinationTexture->GetFlags();
+			bool bIsCurrentSRGB = EnumHasAnyFlags(CurrentFlags, TexCreate_SRGB);
+			if (bIsCurrentSRGB == bIsDestinationTextureSRGB)
+			{
+				return DestinationTexture;
+			}
 		}
-
+		
 		FRHIResourceCreateInfo CreateInfo(TEXT("FWmfMediaHardwareVideoDecodingTextureSample_DestinationTexture"));
-		const ETextureCreateFlags CreateFlags = TexCreate_Dynamic | TexCreate_DisableSRVCreation;
+		const ETextureCreateFlags CreateFlags = TexCreate_Dynamic | TexCreate_DisableSRVCreation |
+			(bIsDestinationTextureSRGB ? TexCreate_SRGB : TexCreate_None);
 		DestinationTexture = RHICreateTexture2D(
 			Dim.X,
 			Dim.Y,
@@ -114,6 +141,33 @@ public:
 		return DestinationTexture;
 	}
 
+	/**
+	 * Get Destination Texture of render thread device
+	 *
+	 * @return Destination texture
+	 */
+	FTexture2DRHIRef GetOrCreateDestinationAlphaTexture()
+	{
+		if (DestinationAlphaTexture.IsValid() && DestinationAlphaTexture->GetSizeX() == Dim.X && DestinationAlphaTexture->GetSizeY() == Dim.Y)
+		{
+			return DestinationAlphaTexture;
+		}
+		
+		FRHIResourceCreateInfo CreateInfo(TEXT("FWmfMediaHardwareVideoDecodingTextureSample_DestinationAlphaTexture"));
+		const ETextureCreateFlags CreateFlags = TexCreate_Dynamic | TexCreate_DisableSRVCreation;
+		DestinationAlphaTexture = RHICreateTexture2D(
+			Dim.X,
+			Dim.Y,
+			AlphaFormat,
+			1,
+			1,
+			CreateFlags,
+			CreateInfo);
+
+		return DestinationAlphaTexture;
+	}
+
+
 	/**
 	 * Called the the sample is returned to the pool for cleanup purposes
 	 */
@@ -132,8 +186,17 @@ private:
 	/** Destination Texture resource (from Rendering device) */
 	FTexture2DRHIRef DestinationTexture;
 
+	/** Destination Texture resource (from Rendering device) */
+	FTexture2DRHIRef DestinationAlphaTexture;
+
 	/** Texture format */
 	EPixelFormat Format;
+
+	/** Texture format */
+	EPixelFormat AlphaFormat;
+
+	/** Whether DestinationTexture should be an sRGB texture or not. */
+	bool bIsDestinationTextureSRGB;
 };
 
 /** Implements a pool for WMF texture samples. */
diff --git a/Engine/Source/Runtime/Media/Public/IMediaTextureSample.h b/Engine/Source/Runtime/Media/Public/IMediaTextureSample.h
--- a/Engine/Source/Runtime/Media/Public/IMediaTextureSample.h
+++ b/Engine/Source/Runtime/Media/Public/IMediaTextureSample.h
@@ -72,6 +72,18 @@ enum class EMediaTextureSampleFormat
 
 	/** 4:4:4:4 AY'CbCr 16-bit little endian full range alpha, video range Y'CbCr. */
 	Y416,
+
+	/** DXT1. */
+	DXT1,
+
+	/** DXT5. */
+	DXT5,
+
+	/** YCoCg colour space encoded in DXT5. */
+	YCoCg_DXT5,
+
+	/** YCoCg colour space encoded in DXT5, with a separate alpha texture encoded in BC4. */
+	YCoCg_DXT5_Alpha_BC4,
 };
 
 namespace MediaTextureSampleFormat
