Fixing up DX12 playback with WMFmediaPlayer (H264/5, HAP, ProRes)

diff --git a/Engine/Plugins/Media/AppleProResMedia/Source/AppleProResMedia/Private/AppleProResDecoder/WmfMediaAppleProResDecoder.cpp b/Engine/Plugins/Media/AppleProResMedia/Source/AppleProResMedia/Private/AppleProResDecoder/WmfMediaAppleProResDecoder.cpp
--- a/Engine/Plugins/Media/AppleProResMedia/Source/AppleProResMedia/Private/AppleProResDecoder/WmfMediaAppleProResDecoder.cpp
+++ b/Engine/Plugins/Media/AppleProResMedia/Source/AppleProResMedia/Private/AppleProResDecoder/WmfMediaAppleProResDecoder.cpp
@@ -12,8 +12,6 @@
 
 #include "Windows/AllowWindowsPlatformTypes.h"
 
-#include <d3d11.h>
-
 const GUID DecoderGUID_AppleProRes_422_Proxy = { 0x6170636F, 0x767A, 0x494D, { 0xB4, 0x78, 0xF2, 0x9D, 0x25, 0xDC, 0x90, 0x37 } };
 const GUID DecoderGUID_AppleProRes_422_LT = { 0x61706373, 0x767A, 0x494D, { 0xB4, 0x78, 0xF2, 0x9D, 0x25, 0xDC, 0x90, 0x37 } };
 const GUID DecoderGUID_AppleProRes_422 = { 0x6170636E, 0x767A, 0x494D, { 0xB4, 0x78, 0xF2, 0x9D, 0x25, 0xDC, 0x90, 0x37 } };
@@ -185,36 +183,12 @@ HRESULT WmfMediaAppleProResDecoder::ProcessMessage(MFT_MESSAGE_TYPE eMessage, UL
 		hr = OnDiscontinuity();
 		break;
 
-	case MFT_MESSAGE_SET_D3D_MANAGER:
-	{
-		IMFDXGIDeviceManager *pDeviceManager = (IMFDXGIDeviceManager*)ulParam;
-		DXGIManager = TComPtr<IMFDXGIDeviceManager>(pDeviceManager);
-
-		if (DXGIManager.IsValid())
-		{
-			HANDLE DeviceHandle = 0;
-			if (DXGIManager->OpenDeviceHandle(&DeviceHandle) == S_OK)
-			{
-				DXGIManager->GetVideoService(DeviceHandle, __uuidof(ID3D11Device), (void**)&D3D11Device);
-				if (D3D11Device.IsValid())
-				{
-					D3D11Device->GetImmediateContext(&D3DImmediateContext);
-					if (D3DImmediateContext.IsValid())
-					{
-						UE_LOG(LogAppleProResMedia, Verbose, TEXT("D3D11Device from Device manager: %p"), D3D11Device.Get());
-					}
-				}
-			}
-		}
-	}
-		break;
 	default:
 		break;
 	}
 
 	return hr;
 }
-	
 
 
 HRESULT WmfMediaAppleProResDecoder::ProcessOutput(DWORD dwFlags, DWORD cOutputBufferCount, MFT_OUTPUT_DATA_BUFFER* pOutputSamples, DWORD* pdwStatus)
@@ -239,32 +213,10 @@ HRESULT WmfMediaAppleProResDecoder::ProcessOutput(DWORD dwFlags, DWORD cOutputBu
 		return MF_E_TRANSFORM_NEED_MORE_INPUT;
 	}
 
-	if (!OutputTexture.IsValid())
-	{
-		D3D11_TEXTURE2D_DESC TextureDesc;
-		memset(&TextureDesc, 0, sizeof(TextureDesc));
-		TextureDesc.Width = ImageWidthInPixels;
-		TextureDesc.Height = ImageHeightInPixels;
-		TextureDesc.MipLevels = 1;
-		TextureDesc.ArraySize = 1;
-		TextureDesc.Format = DXGI_FORMAT_R16G16B16A16_UNORM;
-		TextureDesc.SampleDesc.Count = 1;
-		TextureDesc.SampleDesc.Quality = 0;
-		TextureDesc.Usage = D3D11_USAGE_DYNAMIC;
-		TextureDesc.BindFlags = D3D11_BIND_SHADER_RESOURCE;
-		TextureDesc.CPUAccessFlags = D3D11_CPU_ACCESS_WRITE;
-		TextureDesc.MiscFlags = 0;
-		HRESULT Result = D3D11Device->CreateTexture2D(&TextureDesc, nullptr, &OutputTexture);
-		printf("%ld", Result);
-	}
-
-
 	FScopeLock Lock(&CriticalSection);
 
 	HRESULT hr = S_OK;
 
-	TComPtr<IMFMediaBuffer> pOutput;
-
 	if (SUCCEEDED(hr))
 	{
 		TComPtr<IMFSample> Sample;
@@ -274,20 +226,16 @@ HRESULT WmfMediaAppleProResDecoder::ProcessOutput(DWORD dwFlags, DWORD cOutputBu
 		pOutputSamples[0].pSample = Sample.Get();
 		Sample->AddRef();
 
-		TComPtr<IMFMediaBuffer> MediaBuffer;
-		MFCreateDXGISurfaceBuffer(__uuidof(ID3D11Texture2D), OutputTexture, 0, FALSE, &MediaBuffer);
+		// If used with DX11 we need a (dummy, but amply sized!) memory buffer attached to the sample to make WMF happy
+		if (!MediaBuffer.IsValid())
+		{
+			MFCreateMemoryBuffer(OutputImageSize, &MediaBuffer);
+			MediaBuffer->SetCurrentLength(OutputImageSize);
+		}
 
 		Sample->AddBuffer(MediaBuffer);
 	}
 
-
-	hr = pOutputSamples[0].pSample->GetBufferByIndex(0, &pOutput);
-
-	if (pOutput == nullptr)
-	{
-		return S_OK;
-	}
-
 	if (SUCCEEDED(hr))
 	{
 		hr = InternalProcessOutput(pOutputSamples[0].pSample);
@@ -311,87 +259,27 @@ HRESULT WmfMediaAppleProResDecoder::InternalProcessOutput(IMFSample* InSample)
 		return MF_E_TRANSFORM_NEED_MORE_INPUT;
 	}
 
-	HRESULT hr = S_OK;
-
-	BYTE *pbData = NULL;
-
-	DWORD dwTimeCode = 0;
-	LONGLONG rt = 0;
-
-	DWORD cBuffers = 0;
-	TComPtr<IMFMediaBuffer> pBuffer = nullptr;
-	TComPtr<IMFDXGIBuffer> pDXGIBuffer = nullptr;
-	UINT dwViewIndex = 0;
-	TComPtr<ID3D11Texture2D> TextureNV12_Y = nullptr;
-
-	HRESULT Result = S_OK;
-	if (bIsExternalBufferEnabled == false)
-	{
-		Result = InSample->GetBufferCount(&cBuffers);
-		if (FAILED(Result))
-		{
-			return Result;
-		}
-		if (1 == cBuffers)
-		{
-			Result = InSample->GetBufferByIndex(0, &pBuffer);
-			if (FAILED(Result))
-			{
-				return Result;
-			}
-		}
-		Result = pBuffer->QueryInterface(__uuidof(IMFDXGIBuffer), (LPVOID*)&pDXGIBuffer);
-		if (FAILED(Result))
-		{
-			return Result;
-		}
-
-		Result = pDXGIBuffer->GetResource(__uuidof(ID3D11Texture2D), (LPVOID*)&TextureNV12_Y);
-		if (FAILED(Result))
-		{
-			return Result;
-		}
-
-		Result = pDXGIBuffer->GetSubresourceIndex(&dwViewIndex);
-		if (FAILED(Result))
-		{
-			return Result;
-		}
-
-	}
-
 	DataBuffer OuputDataBuffer;
 	OutputQueue.Dequeue(OuputDataBuffer);
 
 	LONGLONG TimeStamp = OuputDataBuffer.TimeStamp;
 
-	if (bIsExternalBufferEnabled == false)
+	check(bIsExternalBufferEnabled);
+
+	// Get buffer.
+	TArray<uint8>* pExternalBuffer = AllocateExternalBuffer(TimeStamp, OuputDataBuffer.Color.Num());
+	if (pExternalBuffer == nullptr)
 	{
-		D3D11_MAPPED_SUBRESOURCE MappedResourceColor;
-		Result = D3DImmediateContext->Map(OutputTexture, 0, D3D11_MAP_WRITE_DISCARD, 0, &MappedResourceColor);
-		memcpy(MappedResourceColor.pData, OuputDataBuffer.Color.GetData(), MappedResourceColor.DepthPitch);
-		D3DImmediateContext->Unmap(OutputTexture, 0);
+		return 0;
 	}
-	else
-	{
-		// Get buffer.
-		TArray<uint8>* pExternalBuffer = AllocateExternalBuffer(TimeStamp, OuputDataBuffer.Color.Num());
-		if (pExternalBuffer == nullptr)
-		{
-			return 0;
-		}
 		
-		// Copy to buffer.
-		pExternalBuffer->SetNum(OuputDataBuffer.Color.Num());
-		FMemory::Memcpy(pExternalBuffer->GetData(), OuputDataBuffer.Color.GetData(), OuputDataBuffer.Color.Num());
-	}
+	// Copy to buffer.
+	pExternalBuffer->SetNum(OuputDataBuffer.Color.Num());
+	FMemory::Memcpy(pExternalBuffer->GetData(), OuputDataBuffer.Color.GetData(), OuputDataBuffer.Color.Num());
 
 	InputQueue.Enqueue(MoveTemp(OuputDataBuffer));
 
-	if (SUCCEEDED(hr))
-	{
-		hr = InSample->SetUINT32(MFSampleExtension_CleanPoint, TRUE);
-	}
+	HRESULT hr = InSample->SetUINT32(MFSampleExtension_CleanPoint, TRUE);
 
 	if (SUCCEEDED(hr))
 	{
@@ -490,7 +378,7 @@ HRESULT WmfMediaAppleProResDecoder::OnSetInputType(IMFMediaType* InMediaType)
 	if (SUCCEEDED(hr))
 	{
 		InputImageSize = ImageWidthInPixels * ImageHeightInPixels * 8;
-		OutputImageSize = ImageWidthInPixels * ImageHeightInPixels * 4;
+		OutputImageSize = ImageWidthInPixels * ImageHeightInPixels * 8; // Y416 (4.4.4 - 16-bit - aligned)
 
 		InputType = InMediaType;
 		InputType->AddRef();
@@ -550,4 +438,3 @@ HRESULT WmfMediaAppleProResDecoder::InternalProcessInput(LONGLONG InTimeStamp, B
 #include "Windows/HideWindowsPlatformTypes.h"
 
 #endif // WMFMEDIA_SUPPORTED_PLATFORM
-
diff --git a/Engine/Plugins/Media/AppleProResMedia/Source/AppleProResMedia/Private/AppleProResDecoder/WmfMediaAppleProResDecoder.h b/Engine/Plugins/Media/AppleProResMedia/Source/AppleProResMedia/Private/AppleProResDecoder/WmfMediaAppleProResDecoder.h
--- a/Engine/Plugins/Media/AppleProResMedia/Source/AppleProResMedia/Private/AppleProResDecoder/WmfMediaAppleProResDecoder.h
+++ b/Engine/Plugins/Media/AppleProResMedia/Source/AppleProResMedia/Private/AppleProResDecoder/WmfMediaAppleProResDecoder.h
@@ -43,8 +43,7 @@ private:
 protected:
 	
 	PRDecoderRef Decoder;
-
-	TComPtr<ID3D11Texture2D> OutputTexture;
+	TComPtr<IMFMediaBuffer> MediaBuffer;
 };
 
 #endif // WMFMEDIA_SUPPORTED_PLATFORM
diff --git a/Engine/Plugins/Media/HAPMedia/Source/HAPMedia/Private/HAPDecoder/WmfMediaHAPDecoder.cpp b/Engine/Plugins/Media/HAPMedia/Source/HAPMedia/Private/HAPDecoder/WmfMediaHAPDecoder.cpp
--- a/Engine/Plugins/Media/HAPMedia/Source/HAPMedia/Private/HAPDecoder/WmfMediaHAPDecoder.cpp
+++ b/Engine/Plugins/Media/HAPMedia/Source/HAPMedia/Private/HAPDecoder/WmfMediaHAPDecoder.cpp
@@ -4,6 +4,8 @@
 
 #if WMFMEDIA_SUPPORTED_PLATFORM
 
+#include "WmfMediaCommon.h"
+
 #include "hap.h"
 
 #include "HAPMediaModule.h"
@@ -12,9 +14,6 @@
 #include "Misc/Paths.h"
 #include "Windows/AllowWindowsPlatformTypes.h"
 
-#include <d3d11.h>
-#include <d3dcompiler.h>
-
 void MyHapDecodeCallback(HapDecodeWorkFunction InFunction, void* InParameter, unsigned int InCount, void* InInfo)
 {
 	unsigned int i;
@@ -59,8 +58,7 @@ bool WmfMediaHAPDecoder::SetOutputFormat(const GUID& InGuid, GUID& OutVideoForma
 
 WmfMediaHAPDecoder::WmfMediaHAPDecoder()
 	: WmfMediaDecoder(),
-	InputSubType{ 0 },
-	FrameData{ 0 }
+	InputSubType{ 0 }
 {
 }
 
@@ -126,34 +124,25 @@ HRESULT WmfMediaHAPDecoder::GetOutputAvailableType(DWORD dwOutputStreamID, DWORD
 
 	if (SUCCEEDED(hr))
 	{
-		if (bIsExternalBufferEnabled)
+		check(bIsExternalBufferEnabled);
+		hr = TempOutputType->SetGUID(MF_MT_SUBTYPE, MFVideoFormat_NV12);
+		if (SUCCEEDED(hr))
 		{
 			if (InputSubType == DecoderGUID_HAP)
 			{
-				hr = TempOutputType->SetGUID(MF_MT_SUBTYPE, MFVideoFormat_L8);
+				hr = TempOutputType->SetGUID(UE_WMF_PrivateFormatGUID, MFVideoFormat_L8);
 			}
 			else if (InputSubType == DecoderGUID_HAP_ALPHA)
 			{
-				hr = TempOutputType->SetGUID(MF_MT_SUBTYPE, MFVideoFormat_L16);
+				hr = TempOutputType->SetGUID(UE_WMF_PrivateFormatGUID, MFVideoFormat_L16);
 			}
 			else if (InputSubType == DecoderGUID_HAP_Q)
 			{
-				hr = TempOutputType->SetGUID(MF_MT_SUBTYPE, MFVideoFormat_RGB8);
-			}
-			else
-			{
-				hr = TempOutputType->SetGUID(MF_MT_SUBTYPE, MFVideoFormat_D16);
-			}
-		}
-		else
-		{
-			if (InputSubType == DecoderGUID_HAP_ALPHA || InputSubType == DecoderGUID_HAP_Q_ALPHA)
-			{
-				hr = TempOutputType->SetGUID(MF_MT_SUBTYPE, MFVideoFormat_ARGB32);
+				hr = TempOutputType->SetGUID(UE_WMF_PrivateFormatGUID, MFVideoFormat_RGB8);
 			}
 			else
 			{
-				hr = TempOutputType->SetGUID(MF_MT_SUBTYPE, MFVideoFormat_NV12);
+				hr = TempOutputType->SetGUID(UE_WMF_PrivateFormatGUID, MFVideoFormat_D16);
 			}
 		}
 	}
@@ -203,6 +192,27 @@ HRESULT WmfMediaHAPDecoder::GetOutputAvailableType(DWORD dwOutputStreamID, DWORD
 }
 
 
+
+HRESULT WmfMediaHAPDecoder::GetAttributes(IMFAttributes** ppAttributes)
+{
+	TComPtr<IMFAttributes> Attributes;
+	HRESULT hr = MFCreateAttributes(&Attributes, 1);
+
+	if (SUCCEEDED(hr))
+	{
+		const GUID NoneEmptyDummyGUID = { 0x86c8d2ec, 0xa1e, 0x4637, {0xbe, 0x7a, 0x7f, 0xea, 0x0, 0xb2, 0xb3, 0xb1} };
+		Attributes->SetUINT32(NoneEmptyDummyGUID, 0U);
+		*ppAttributes = Attributes;
+		(*ppAttributes)->AddRef();
+		return S_OK;
+	}
+	else
+	{
+		return E_NOTIMPL;
+	}
+}
+
+
 HRESULT WmfMediaHAPDecoder::ProcessMessage(MFT_MESSAGE_TYPE eMessage, ULONG_PTR ulParam)
 {
 	FScopeLock Lock(&CriticalSection);
@@ -220,34 +230,6 @@ HRESULT WmfMediaHAPDecoder::ProcessMessage(MFT_MESSAGE_TYPE eMessage, ULONG_PTR
 		hr = OnDiscontinuity();
 		break;
 
-	case MFT_MESSAGE_SET_D3D_MANAGER:
-		{
-			IMFDXGIDeviceManager *pDeviceManager = (IMFDXGIDeviceManager*)ulParam;
-			DXGIManager = TComPtr<IMFDXGIDeviceManager>(pDeviceManager);
-
-			if (DXGIManager.IsValid())
-			{
-				HANDLE DeviceHandle = 0;
-				if (DXGIManager->OpenDeviceHandle(&DeviceHandle) == S_OK)
-				{
-					DXGIManager->GetVideoService(DeviceHandle, __uuidof(ID3D11Device), (void**)&D3D11Device);
-					if (D3D11Device.IsValid())
-					{
-						D3D11Device->GetImmediateContext(&D3DImmediateContext);
-						if (D3DImmediateContext.IsValid())
-						{
-							UE_LOG(LogHAPMedia, Verbose, TEXT("D3D11Device from Device manager: %p"), D3D11Device.Get());
-							if (!InitPipeline())
-							{
-								UE_LOG(LogHAPMedia, Error, TEXT("Unable to initialize 3D pipeline for Hap Decoding"));
-								hr = E_NOTIMPL;
-							}
-						}
-					}
-				}
-			}
-		}
-		break;
 	default:
 		break;
 	}
@@ -278,141 +260,10 @@ HRESULT WmfMediaHAPDecoder::ProcessOutput(DWORD dwFlags, DWORD cOutputBufferCoun
 		return MF_E_TRANSFORM_NEED_MORE_INPUT;
 	}
 
-	if (!D3D11Device.IsValid())
-	{
-		UE_LOG(LogHAPMedia, Error, TEXT("DX11 Device is not initialized!"));
-		return E_NOTIMPL;
-	}
-
-	if (!OutputTexture.IsValid())
-	{
-		TComPtr<ID3D11Texture2D> Texture;
-		TComPtr<ID3D11RenderTargetView> RenderTarget_A;
-		TComPtr<ID3D11RenderTargetView> RenderTarget_B;
-		{
-			D3D11_TEXTURE2D_DESC TextureDesc;
-			memset(&TextureDesc, 0, sizeof(TextureDesc));
-			TextureDesc.Width = ImageWidthInPixels;
-			TextureDesc.Height = ImageHeightInPixels;
-			TextureDesc.MipLevels = 1;
-			TextureDesc.ArraySize = 1;
-
-			if (InputSubType == DecoderGUID_HAP || InputSubType == DecoderGUID_HAP_Q)
-			{
-				TextureDesc.Format = DXGI_FORMAT_NV12;
-			}
-			else // if (InputSubType == DecoderGUID_HAP_ALPHA || InputSubType == DecoderGUID_HAP_ALPHA)
-			{
-				TextureDesc.Format = DXGI_FORMAT_B8G8R8A8_UNORM;
-			}
-
-			TextureDesc.SampleDesc.Count = 1;
-			TextureDesc.SampleDesc.Quality = 0;
-			TextureDesc.Usage = D3D11_USAGE_DEFAULT;
-			TextureDesc.BindFlags = D3D11_BIND_RENDER_TARGET;
-			TextureDesc.CPUAccessFlags = 0;
-			TextureDesc.MiscFlags = 0;
-			HRESULT Result = D3D11Device->CreateTexture2D(&TextureDesc, nullptr, &Texture);
-
-			if (InputSubType == DecoderGUID_HAP || InputSubType == DecoderGUID_HAP_Q)
-			{
-				D3D11_RENDER_TARGET_VIEW_DESC RTViewDesc;
-				memset(&RTViewDesc, 0, sizeof(RTViewDesc));
-				RTViewDesc.Format = DXGI_FORMAT_R8_UNORM;
-				RTViewDesc.ViewDimension = D3D11_RTV_DIMENSION_TEXTURE2D;
-				RTViewDesc.Texture2D.MipSlice = 0;
-				Result = D3D11Device->CreateRenderTargetView(Texture, &RTViewDesc, &RenderTarget_A);
-
-				memset(&RTViewDesc, 0, sizeof(RTViewDesc));
-				RTViewDesc.Format = DXGI_FORMAT_R8G8_UNORM;
-				RTViewDesc.ViewDimension = D3D11_RTV_DIMENSION_TEXTURE2D;
-				RTViewDesc.Texture2D.MipSlice = 0;
-				Result = D3D11Device->CreateRenderTargetView(Texture, &RTViewDesc, &RenderTarget_B);
-			}
-			else // if (InputSubType == DecoderGUID_HAP_ALPHA || InputSubType == DecoderGUID_HAP_Q_ALPHA)
-			{
-				D3D11_RENDER_TARGET_VIEW_DESC RTViewDesc;
-				memset(&RTViewDesc, 0, sizeof(RTViewDesc));
-				RTViewDesc.Format = TextureDesc.Format;
-				RTViewDesc.ViewDimension = D3D11_RTV_DIMENSION_TEXTURE2D;
-				RTViewDesc.Texture2D.MipSlice = 0;
-				Result = D3D11Device->CreateRenderTargetView(Texture, &RTViewDesc, &RenderTarget_A);
-			}
-		}
-
-		OutputTexture = MakeUnique<FOutputTextures>();
-		OutputTexture->Texture = MoveTemp(Texture);
-		OutputTexture->RTV_A = MoveTemp(RenderTarget_A);
-		OutputTexture->RTV_B = MoveTemp(RenderTarget_B);
-	}
-
-	if (!WorkTextures.IsValid())
-	{
-		TComPtr<ID3D11Texture2D> InputTextureColor;
-		TComPtr<ID3D11ShaderResourceView> SRV_InputColor;
-		{
-			D3D11_TEXTURE2D_DESC TextureDesc;
-			memset(&TextureDesc, 0, sizeof(TextureDesc));
-			TextureDesc.Width = ImageWidthInPixels;
-			TextureDesc.Height = ImageHeightInPixels;
-			TextureDesc.MipLevels = 1;
-			TextureDesc.ArraySize = 1;
-
-			if (InputSubType == DecoderGUID_HAP)
-			{
-				TextureDesc.Format = DXGI_FORMAT_BC1_UNORM;
-			}
-			else // if (InputSubType == DecoderGUID_HAP_ALPHA || InputSubType == DecoderGUID_HAP_Q || InputSubType == DecoderGUID_HAP_Q_ALPHA)
-			{
-				TextureDesc.Format = DXGI_FORMAT_BC3_UNORM;
-			}
-
-			TextureDesc.SampleDesc.Count = 1;
-			TextureDesc.SampleDesc.Quality = 0;
-			TextureDesc.Usage = D3D11_USAGE_DYNAMIC;
-			TextureDesc.BindFlags = D3D11_BIND_SHADER_RESOURCE;
-			TextureDesc.CPUAccessFlags = D3D11_CPU_ACCESS_WRITE;
-			TextureDesc.MiscFlags = 0;
-			D3D11Device->CreateTexture2D(&TextureDesc, nullptr, &InputTextureColor);
-			CreateSRV(InputTextureColor, TextureDesc, &SRV_InputColor);
-		}
-
-		TComPtr<ID3D11Texture2D> InputTextureAlpha;
-		TComPtr<ID3D11ShaderResourceView> SRV_InputAlpha;
-
-		if (InputSubType == DecoderGUID_HAP_Q_ALPHA)
-		{
-			D3D11_TEXTURE2D_DESC TextureDesc;
-			memset(&TextureDesc, 0, sizeof(TextureDesc));
-			TextureDesc.Width = ImageWidthInPixels;
-			TextureDesc.Height = ImageHeightInPixels;
-			TextureDesc.MipLevels = 1;
-			TextureDesc.ArraySize = 1;
-			TextureDesc.Format = DXGI_FORMAT_BC4_UNORM;
-			TextureDesc.SampleDesc.Count = 1;
-			TextureDesc.SampleDesc.Quality = 0;
-			TextureDesc.Usage = D3D11_USAGE_DYNAMIC;
-			TextureDesc.BindFlags = D3D11_BIND_SHADER_RESOURCE;
-			TextureDesc.CPUAccessFlags = D3D11_CPU_ACCESS_WRITE;
-			TextureDesc.MiscFlags = 0;
-			D3D11Device->CreateTexture2D(&TextureDesc, nullptr, &InputTextureAlpha);
-			CreateSRV(InputTextureAlpha, TextureDesc, &SRV_InputAlpha);
-		}
-
-		FWorkTextures WT;
-		WT.InputTextureColor = MoveTemp(InputTextureColor);
-		WT.InputTextureAlpha = MoveTemp(InputTextureAlpha);
-		WT.SRV_InputColor = MoveTemp(SRV_InputColor);
-		WT.SRV_InputAlpha = MoveTemp(SRV_InputAlpha);
-		WorkTextures = MakeUnique<FWorkTextures>(MoveTemp(WT));
-	}
-
 	FScopeLock Lock(&CriticalSection);
 
 	HRESULT hr = S_OK;
 
-	TComPtr<IMFMediaBuffer> pOutput;
-
 	if (SUCCEEDED(hr))
 	{
 		TComPtr<IMFSample> Sample;
@@ -422,19 +273,16 @@ HRESULT WmfMediaHAPDecoder::ProcessOutput(DWORD dwFlags, DWORD cOutputBufferCoun
 		pOutputSamples[0].pSample = Sample.Get();
 		Sample->AddRef();
 
-		TComPtr<IMFMediaBuffer> MediaBuffer;
-		MFCreateDXGISurfaceBuffer(__uuidof(ID3D11Texture2D), OutputTexture->Texture, 0, FALSE, &MediaBuffer);
+		// If used with DX11 we need a (dummy, but amply sized!) memory buffer attached to the sample to make WMF happy
+		if (!MediaBuffer.IsValid())
+		{
+			MFCreateMemoryBuffer(OutputImageSize, &MediaBuffer);
+			MediaBuffer->SetCurrentLength(OutputImageSize);
+		}
 
 		Sample->AddBuffer(MediaBuffer);
 	}
 
-	hr = pOutputSamples[0].pSample->GetBufferByIndex(0, &pOutput);
-
-	if (pOutput == nullptr)
-	{
-		return S_OK;
-	}
-
 	if (SUCCEEDED(hr))
 	{
 		hr = InternalProcessOutput(pOutputSamples[0].pSample);
@@ -460,183 +308,49 @@ HRESULT WmfMediaHAPDecoder::InternalProcessOutput(IMFSample* InSample)
 		return MF_E_TRANSFORM_NEED_MORE_INPUT;
 	}
 
-	HRESULT hr = S_OK;
+	DataBuffer OuputDataBuffer;
+	OutputQueue.Dequeue(OuputDataBuffer);
 
-	BYTE *pbData = NULL;
-	LONG lActualStride = 0;
+	LONGLONG TimeStamp = OuputDataBuffer.TimeStamp;
 
-	DWORD dwTimeCode = 0;
-	LONGLONG rt = 0;
+	// Are we using external buffers?
+	check(bIsExternalBufferEnabled);
 
-	DWORD cBuffers = 0;
-	TComPtr<IMFMediaBuffer> pBuffer = nullptr;
-	TComPtr<IMFDXGIBuffer> pDXGIBuffer = nullptr;
-	UINT dwViewIndex = 0;
-	TComPtr<ID3D11Texture2D> TextureNV12_Y = nullptr;
+	// Get buffer.
+	int32 BufferSize = ImageWidthInPixels * ImageHeightInPixels;
+	int32 AlphaBufferSize = 0;
 
-	HRESULT Result = InSample->GetBufferCount(&cBuffers);
-	if (FAILED(Result))
-	{
-		return Result;
-	}
-	if (1 == cBuffers)
+	// HapQ Alpha has a separate alpha buffer.
+	if (InputSubType == DecoderGUID_HAP_Q_ALPHA)
 	{
-		Result = InSample->GetBufferByIndex(0, &pBuffer);
-		if (FAILED(Result))
-		{
-			return Result;
-		}
-	}
-	Result = pBuffer->QueryInterface(__uuidof(IMFDXGIBuffer), (LPVOID*)&pDXGIBuffer);
-	if (FAILED(Result))
-	{
-		return Result;
+		AlphaBufferSize = BufferSize / 2;
 	}
 
-	Result = pDXGIBuffer->GetResource(__uuidof(ID3D11Texture2D), (LPVOID*)&TextureNV12_Y);
-	if (FAILED(Result))
+	// Hap is 4 bits per pixel.
+	if (InputSubType == DecoderGUID_HAP)
 	{
-		return Result;
+		BufferSize = BufferSize / 2;
 	}
 
-	Result = pDXGIBuffer->GetSubresourceIndex(&dwViewIndex);
-	if (FAILED(Result))
+	// Get external buffer.
+	TArray<uint8>* pExternalBuffer = AllocateExternalBuffer(TimeStamp, BufferSize + AlphaBufferSize);
+	if (pExternalBuffer == nullptr)
 	{
-		return Result;
+		return 0;
 	}
 
-	DataBuffer OuputDataBuffer;
-	OutputQueue.Dequeue(OuputDataBuffer);
-
-	LONGLONG TimeStamp = OuputDataBuffer.TimeStamp;
+	// Copy to buffer.
+	pExternalBuffer->SetNum(BufferSize + AlphaBufferSize);
+	FMemory::Memcpy(pExternalBuffer->GetData(), OuputDataBuffer.Color.GetData(), BufferSize);
 
-	// Are we using external buffers?
-	if (bIsExternalBufferEnabled == false)
+	if (AlphaBufferSize > 0)
 	{
-		D3D11_MAPPED_SUBRESOURCE MappedResourceColor;
-		Result = D3DImmediateContext->Map(WorkTextures->InputTextureColor, 0, D3D11_MAP_WRITE_DISCARD, 0, &MappedResourceColor);
-		memcpy(MappedResourceColor.pData, OuputDataBuffer.Color.GetData(), MappedResourceColor.DepthPitch);
-		D3DImmediateContext->Unmap(WorkTextures->InputTextureColor, 0);
-
-		if (InputSubType == DecoderGUID_HAP_Q_ALPHA)
-		{
-			D3D11_MAPPED_SUBRESOURCE MappedResourceAlpha;
-			Result = D3DImmediateContext->Map(WorkTextures->InputTextureAlpha, 0, D3D11_MAP_WRITE_DISCARD, 0, &MappedResourceAlpha);
-			memcpy(MappedResourceAlpha.pData, OuputDataBuffer.Alpha.GetData(), MappedResourceAlpha.DepthPitch);
-			D3DImmediateContext->Unmap(WorkTextures->InputTextureAlpha, 0);
-		}
-
-
-		if (InputSubType == DecoderGUID_HAP_Q_ALPHA)
-		{
-			BindSRV(WorkTextures->SRV_InputColor, WorkTextures->SRV_InputAlpha);
-		}
-		else
-		{
-			BindSRV(WorkTextures->SRV_InputColor);
-		}
-
-		UINT stride = sizeof(VertexDescription);
-		UINT offset = 0;
-
-		if (InputSubType == DecoderGUID_HAP)
-		{
-			FrameData.Mode = 0;
-		}
-		else if (InputSubType == DecoderGUID_HAP_ALPHA)
-		{
-			FrameData.Mode = 1;
-		}
-		else if (InputSubType == DecoderGUID_HAP_Q)
-		{
-			FrameData.Mode = 2;
-		}
-		else // if (InputSubType == DecoderGUID_HAP_Q_ALPHA)
-		{
-			FrameData.Mode = 3;
-		}
-
-		D3D11_VIEWPORT viewport;
-		ZeroMemory(&viewport, sizeof(D3D11_VIEWPORT));
-		viewport.TopLeftX = 0;
-		viewport.TopLeftY = 0;
-
-		ID3D11RenderTargetView* m_pRenderViews[1];
-
-		int StartPass = 0;
-
-		if (InputSubType == DecoderGUID_HAP_ALPHA || InputSubType == DecoderGUID_HAP_Q_ALPHA)
-		{
-			StartPass = 1;
-		}
-
-		for (int PassIndex = StartPass; PassIndex < 2; PassIndex++)
-		{
-			if (PassIndex == 0)
-			{
-				m_pRenderViews[0] = OutputTexture->RTV_B;
-				FrameData.WriteY = 0;
-				viewport.Width = float(ImageWidthInPixels / 2);
-				viewport.Height = float(ImageHeightInPixels / 2);
-			}
-			else
-			{
-				m_pRenderViews[0] = OutputTexture->RTV_A;
-				FrameData.WriteY = 1;
-				viewport.Width = float(ImageWidthInPixels);
-				viewport.Height = float(ImageHeightInPixels);
-			}
-
-			D3DImmediateContext->RSSetViewports(1, &viewport);
-			D3DImmediateContext->UpdateSubresource(PixelBuffer, 0, nullptr, &FrameData, 0, 0);
-			D3DImmediateContext->OMSetRenderTargets(1, &m_pRenderViews[0], NULL);
-			D3DImmediateContext->IASetVertexBuffers(0, 1, &VertexBuffer, &stride, &offset);
-			D3DImmediateContext->IASetPrimitiveTopology(D3D11_PRIMITIVE_TOPOLOGY_TRIANGLESTRIP);
-			D3DImmediateContext->Draw(4, 0);
-		}
-	}
-	else
-	{
-		// Get buffer.
-		int32 BufferSize = ImageWidthInPixels * ImageHeightInPixels;
-		int32 AlphaBufferSize = 0;
-
-		// HapQ Alpha has a separate alpha buffer.
-		if (InputSubType == DecoderGUID_HAP_Q_ALPHA)
-		{
-			AlphaBufferSize = BufferSize / 2;
-		}
-
-		// Hap is 4 bits per pixel.
-		if (InputSubType == DecoderGUID_HAP)
-		{
-			BufferSize = BufferSize / 2;
-		}
-
-		// Get external buffer.
-		TArray<uint8>* pExternalBuffer = AllocateExternalBuffer(TimeStamp, BufferSize + AlphaBufferSize);
-		if (pExternalBuffer == nullptr)
-		{
-			return 0;
-		}
-
-		// Copy to buffer.
-		pExternalBuffer->SetNum(BufferSize + AlphaBufferSize);
-		FMemory::Memcpy(pExternalBuffer->GetData(), OuputDataBuffer.Color.GetData(), BufferSize);
-
-		if (AlphaBufferSize > 0)
-		{
-			FMemory::Memcpy(pExternalBuffer->GetData() + BufferSize, OuputDataBuffer.Alpha.GetData(), AlphaBufferSize);
-		}
+		FMemory::Memcpy(pExternalBuffer->GetData() + BufferSize, OuputDataBuffer.Alpha.GetData(), AlphaBufferSize);
 	}
 	
-
 	InputQueue.Enqueue(MoveTemp(OuputDataBuffer));
 
-	if (SUCCEEDED(hr))
-	{
-		hr = InSample->SetUINT32(MFSampleExtension_CleanPoint, TRUE);
-	}
+	HRESULT hr = InSample->SetUINT32(MFSampleExtension_CleanPoint, TRUE);
 
 	if (SUCCEEDED(hr))
 	{
@@ -738,15 +452,8 @@ HRESULT WmfMediaHAPDecoder::OnSetInputType(IMFMediaType* InMediaType)
 
 	if (SUCCEEDED(hr))
 	{
-		InputImageSize = ImageWidthInPixels * ImageHeightInPixels * 4;
-		if (bIsExternalBufferEnabled)
-		{
-			OutputImageSize = ImageWidthInPixels * ImageHeightInPixels / 2;
-		}
-		else
-		{
-			OutputImageSize = ImageWidthInPixels * ImageHeightInPixels * 3 / 2;
-		}
+		check(bIsExternalBufferEnabled);
+		OutputImageSize = InputImageSize = ImageWidthInPixels * ImageHeightInPixels * 4;
 
 		InputType = InMediaType;
 		InputType->AddRef();
@@ -812,333 +519,7 @@ HRESULT WmfMediaHAPDecoder::InternalProcessInput(LONGLONG InTimeStamp, BYTE* InD
 }
 
 
-bool WmfMediaHAPDecoder::InitPipeline()
-{
-	static const char PixelShader[] = R"SHADER(
-
-	Texture2D TextureColor : register(t0);
-	Texture2D TextureAlpha : register(t1);
-
-	cbuffer cbPerFrameData : register(b0)
-	{
-		int Mode;
-		int Index;
-		int WriteY;
-		int Padding;
-	};
-
-	SamplerState TextureSampler : register(s0);
-
-	float4 PShader(float4 position : SV_POSITION, float2 textureUV : TEXCOORD0) : SV_TARGET
-	{
-		float4 ColorValue = float4(0.0, 0.0, 0.0, 0.0);
-
-		if (Mode == 0)
-		{
-			float4 Value = TextureColor.Sample(TextureSampler, textureUV);
-			float3x3 fMatrix = { 
-				0.299, 0.587, 0.114,
-				-0.14713, -0.28886, 0.436, 
-				0.615, -0.51499, -0.10001 };
-
-			ColorValue = Value;
-			float3 YUV = mul(fMatrix, ColorValue.xyz);
-			YUV += float3(0.06274509803921568627f, 0.5019607843137254902f, 0.5019607843137254902f);
-			ColorValue.rgb = saturate(YUV);
-		}
-		else if (Mode == 1)
-		{
-			float4 Value = TextureColor.Sample(TextureSampler, textureUV);
-			Value = Value * Value;
-			return float4(Value.rgba);
-		}		
-		else
-		{
-			float4 offsets = float4(-0.50196078431373, -0.50196078431373, 0.0, 0.0);
-			float4 CoCgSY = TextureColor.Sample(TextureSampler, textureUV);
-			CoCgSY += offsets;
-
-			float scale = ( CoCgSY.z * ( 255.0 / 8.0 ) ) + 1.0;
-
-			float Co = CoCgSY.x / scale;
-			float Cg = CoCgSY.y / scale;
-			float Y = CoCgSY.w;
-
-			float4 Value = float4(Y + Co - Cg, Y + Cg, Y - Co - Cg, 1.0);
-
-			float3x3 fMatrix = { 
-				0.299, 0.587, 0.114,
-				-0.14713, -0.28886, 0.436, 
-				0.615, -0.51499, -0.10001 };
-
-			ColorValue = Value;
-
-			if (Mode == 3)
-			{
-				float Alpha = TextureAlpha.Sample(TextureSampler, textureUV).r;
-				return float4(ColorValue.rgb*ColorValue.rgb, Alpha);
-			}
-
-			float3 YUV = mul(fMatrix, ColorValue.xyz);
-			YUV += float3(0.06274509803921568627f, 0.5019607843137254902f, 0.5019607843137254902f);
-			ColorValue.rgb = saturate(YUV);
-		}
-
-		if (WriteY == 1)
-		{
-			return float4(ColorValue.x, 0.0, 0.0, 1.0);
-		}
-		else
-		{
-			return float4(ColorValue.y, ColorValue.z, 0.0, 1.0);
-		}
-	}
-	)SHADER";
-
-	static const char VertexShader[] = R"SHADER(
-	struct VOut
-	{
-		float4 position : SV_POSITION;
-		float2 textureUV : TEXCOORD0;
-	};
-
-	VOut VShader(float4 position : POSITION, float2 textureUV : TEXCOORD0)
-	{
-		VOut output;
-		output.position = position;
-		output.textureUV = textureUV;
-		return output;
-	}
-	)SHADER";
-
-	// Keeps compiler stuff together so we can ensure it gets cleaned up appropriately.
-	struct FCompilerVars
-	{
-		~FCompilerVars()
-		{
-			CleanUp();
-		}
-
-		void CleanUp()
-		{
-			if (!bIsDone)
-			{
-				bIsDone = true;
-				PSCode = nullptr;
-				VSCode = nullptr;
-				ErrorMsgs = nullptr;
-				D3DCompileFunc = nullptr;
-				FreeLibrary(CompilerDLL);
-				CompilerDLL = NULL;
-			}
-		}
-
-		TComPtr<ID3DBlob> PSCode;
-		TComPtr<ID3DBlob> VSCode;
-		TComPtr<ID3DBlob> ErrorMsgs;
-
-		HMODULE CompilerDLL = NULL;
-		pD3DCompile D3DCompileFunc = nullptr;
-
-	private:
-		bool bIsDone = false;
-	} CompilerVars;
-
-	// This function needs to load the bundled version of d3dcompiler lib explictly. Implicit linking results
-	// in picking up the system lib by both the engine and SCWs (which links to this module), which makes
-	// the shader compilation system-dependent.
-	FString CompilerPath = FPaths::EngineDir();
-	CompilerPath.Append(TEXT("Binaries/ThirdParty/Windows/DirectX/x64/d3dcompiler_47.dll"));
-	CompilerVars.CompilerDLL = LoadLibrary(*CompilerPath);
-	if (CompilerVars.CompilerDLL == NULL)
-	{
-		UE_LOG(LogHAPMedia, Error, TEXT("Could not locate D3D compiler library at %s"), *CompilerPath);
-		return false;
-	}
-	
-	CompilerVars.D3DCompileFunc = (pD3DCompile)(void*)GetProcAddress(CompilerVars.CompilerDLL, "D3DCompile");
-	if (CompilerVars.D3DCompileFunc == nullptr)
-	{
-		UE_LOG(LogHAPMedia, Error, TEXT("Could not locate D3DCompile() in the compiler library at %s"), *CompilerPath);
-		return false;
-	}
-
-	HRESULT Result = CompilerVars.D3DCompileFunc(PixelShader, sizeof(PixelShader), NULL, NULL, NULL, "PShader", "ps_5_0", 0, 0, &CompilerVars.PSCode, &CompilerVars.ErrorMsgs);
-	if (CompilerVars.ErrorMsgs)
-	{
-		char *pMessage = (char*)CompilerVars.ErrorMsgs->GetBufferPointer();
-		UE_LOG(LogHAPMedia, Error, TEXT("D3DCompile Error/warning: %s"), *FString(pMessage));
-		if (!CompilerVars.PSCode.IsValid())
-		{
-			return false;
-		}
-	}
-
-	Result = CompilerVars.D3DCompileFunc(VertexShader, sizeof(VertexShader), NULL, NULL, NULL, "VShader", "vs_5_0", 0, 0, &CompilerVars.VSCode, &CompilerVars.ErrorMsgs);
-	if (CompilerVars.ErrorMsgs)
-	{
-		char *pMessage = (char*)CompilerVars.ErrorMsgs->GetBufferPointer();
-		UE_LOG(LogHAPMedia, Error, TEXT("D3DCompile Error/warning: %s"), *FString(pMessage));
-		if (!CompilerVars.VSCode.IsValid())
-		{
-			return false;
-		}
-	}
-
-	TComPtr<ID3D11PixelShader> ps;
-	TComPtr<ID3D11VertexShader> vs;
-
-	Result = D3D11Device->CreatePixelShader(CompilerVars.PSCode->GetBufferPointer(), CompilerVars.PSCode->GetBufferSize(), nullptr, &ps);
-	if (FAILED(Result))
-	{
-		return false;
-	}
-
-	Result = D3D11Device->CreateVertexShader(CompilerVars.VSCode->GetBufferPointer(), CompilerVars.VSCode->GetBufferSize(), nullptr, &vs);
-	if (FAILED(Result))
-	{
-		return false;
-	}
-
-	D3D11_BUFFER_DESC PixelConstantBufferDesc;
-	PixelConstantBufferDesc.ByteWidth = sizeof(ConstantBuffer);
-	PixelConstantBufferDesc.Usage = D3D11_USAGE_DEFAULT;
-	PixelConstantBufferDesc.BindFlags = D3D11_BIND_CONSTANT_BUFFER;
-	PixelConstantBufferDesc.CPUAccessFlags = 0;
-	PixelConstantBufferDesc.MiscFlags = 0;
-
-	Result = D3D11Device->CreateBuffer(&PixelConstantBufferDesc, nullptr, &PixelBuffer);
-	if (FAILED(Result))
-	{
-		return false;
-	}
-
-	D3DImmediateContext->PSSetConstantBuffers(0, 1, &PixelBuffer);
-
-	FrameData.Mode = 0;
-	D3DImmediateContext->UpdateSubresource(PixelBuffer, 0, nullptr, &FrameData, 0, 0);
-
-	D3DImmediateContext->VSSetShader(vs, 0, 0);
-	D3DImmediateContext->PSSetShader(ps, 0, 0);
-
-	D3D11_INPUT_ELEMENT_DESC InputElementDescription[] =
-	{
-		{"POSITION", 0, DXGI_FORMAT_R32G32B32_FLOAT, 0, 0, D3D11_INPUT_PER_VERTEX_DATA, 0},
-		{"TEXCOORD", 0, DXGI_FORMAT_R32G32_FLOAT, 0, 12, D3D11_INPUT_PER_VERTEX_DATA, 0},
-	};
-
-	Result = D3D11Device->CreateInputLayout(InputElementDescription, 2, CompilerVars.VSCode->GetBufferPointer(), CompilerVars.VSCode->GetBufferSize(), &InputLayout);
-	if (FAILED(Result))
-	{
-		return false;
-	}
-
-	CompilerVars.CleanUp();
-
-	D3DImmediateContext->IASetInputLayout(InputLayout);
-
-	D3D11_BLEND_DESC BlendStateDesc;
-	ZeroMemory(&BlendStateDesc, sizeof(D3D11_BLEND_DESC));
-	BlendStateDesc.RenderTarget[0].BlendEnable = FALSE;
-	BlendStateDesc.RenderTarget[0].RenderTargetWriteMask = D3D11_COLOR_WRITE_ENABLE_ALL;
-
-	ID3D11BlendState* BlendState = NULL;
-	Result = D3D11Device->CreateBlendState(&BlendStateDesc, &BlendState);
-	if (FAILED(Result))
-	{
-		return false;
-	}
-
-	float BlendFactor[4] = { 1.0f, 1.0f, 1.0f, 1.0f };
-	UINT SampleMask = 0xffffffff;
-	D3DImmediateContext->OMSetBlendState(BlendState, BlendFactor, SampleMask);
-
-	VertexDescription Vertices[] =
-	{
-		{-1.0f,  1.0f, 0.0f, 0.0f, 0.0f },
-		{ 1.0f,  1.0f, 0.0f, 1.0f, 0.0f },
-		{-1.0f, -1.0f, 0.0f, 0.0f, 1.0f },
-		{ 1.0f, -1.0f, 0.0f, 1.0f, 1.0f },
-	};
-
-	D3D11_BUFFER_DESC BufferDesc;
-	ZeroMemory(&BufferDesc, sizeof(BufferDesc));
-	BufferDesc.Usage = D3D11_USAGE_DYNAMIC;
-	BufferDesc.ByteWidth = sizeof(VertexDescription) * 4;
-	BufferDesc.BindFlags = D3D11_BIND_VERTEX_BUFFER;
-	BufferDesc.CPUAccessFlags = D3D11_CPU_ACCESS_WRITE;
-
-	Result = D3D11Device->CreateBuffer(&BufferDesc, NULL, &VertexBuffer);
-	if (FAILED(Result))
-	{
-		return false;
-	}
-
-	D3D11_MAPPED_SUBRESOURCE MappedSubresource;
-	Result = D3DImmediateContext->Map(VertexBuffer, NULL, D3D11_MAP_WRITE_DISCARD, NULL, &MappedSubresource);
-	if (FAILED(Result))
-	{
-		return false;
-	}
-
-	memcpy(MappedSubresource.pData, Vertices, sizeof(Vertices));
-	D3DImmediateContext->Unmap(VertexBuffer, NULL);
-
-	return true;
-}
-
-
-void WmfMediaHAPDecoder::CreateSRV(ID3D11Texture2D* InTexture, D3D11_TEXTURE2D_DESC& InTextureDesc, ID3D11ShaderResourceView** InSRV)
-{
-	D3D11_SHADER_RESOURCE_VIEW_DESC ViewDesc;
-	memset(&ViewDesc, 0, sizeof(ViewDesc));
-
-	ViewDesc.Format = InTextureDesc.Format;
-	ViewDesc.ViewDimension = D3D11_SRV_DIMENSION_TEXTURE2D;
-	ViewDesc.Texture2D.MipLevels = 1;
-	ViewDesc.Texture2D.MostDetailedMip = 0;
-
-	D3D11Device->CreateShaderResourceView(InTexture, &ViewDesc, InSRV);
-
-	D3D11_SAMPLER_DESC SamplerDesc;
-	SamplerDesc.Filter = D3D11_FILTER_MIN_MAG_MIP_LINEAR;
-	SamplerDesc.AddressU = D3D11_TEXTURE_ADDRESS_CLAMP;
-	SamplerDesc.AddressV = D3D11_TEXTURE_ADDRESS_CLAMP;
-	SamplerDesc.AddressW = D3D11_TEXTURE_ADDRESS_CLAMP;
-	SamplerDesc.MipLODBias = 0.0f;
-	SamplerDesc.MaxAnisotropy = 1;
-	SamplerDesc.ComparisonFunc = D3D11_COMPARISON_NEVER;
-	SamplerDesc.MinLOD = -FLT_MAX;
-	SamplerDesc.MaxLOD = FLT_MAX;
-
-	ID3D11SamplerState* myLinearWrapSampler;
-	auto hr = D3D11Device->CreateSamplerState(&SamplerDesc, &myLinearWrapSampler);
-
-	ID3D11SamplerState *SamplerArray[2];
-	SamplerArray[0] = myLinearWrapSampler;
-
-	D3DImmediateContext->PSSetSamplers(0, 1, SamplerArray);
-}
-
-
-void WmfMediaHAPDecoder::BindSRV(ID3D11ShaderResourceView* InSRV)
-{
-	ID3D11ShaderResourceView *Array[2];
-	Array[0] = InSRV;
-	D3DImmediateContext->PSSetShaderResources(0, 1, Array);
-}
-
-
-void WmfMediaHAPDecoder::BindSRV(ID3D11ShaderResourceView* InSRVA, ID3D11ShaderResourceView* InSRVB)
-{
-	ID3D11ShaderResourceView *Array[2];
-	Array[0] = InSRVA;
-	Array[1] = InSRVB;
-	D3DImmediateContext->PSSetShaderResources(0, 2, Array);
-}
-
 
 #include "Windows/HideWindowsPlatformTypes.h"
 
 #endif // WMFMEDIA_SUPPORTED_PLATFORM
-
diff --git a/Engine/Plugins/Media/HAPMedia/Source/HAPMedia/Private/HAPDecoder/WmfMediaHAPDecoder.h b/Engine/Plugins/Media/HAPMedia/Source/HAPMedia/Private/HAPDecoder/WmfMediaHAPDecoder.h
--- a/Engine/Plugins/Media/HAPMedia/Source/HAPMedia/Private/HAPDecoder/WmfMediaHAPDecoder.h
+++ b/Engine/Plugins/Media/HAPMedia/Source/HAPMedia/Private/HAPDecoder/WmfMediaHAPDecoder.h
@@ -6,19 +6,9 @@
 
 #if WMFMEDIA_SUPPORTED_PLATFORM
 
-struct ID3D11Texture2D;
-struct D3D11_TEXTURE2D_DESC;
-struct ID3D11InputLayout;
-struct ID3D11VertexShader;
-struct ID3D11PixelShader;
-struct ID3D11Buffer;
-struct ID3D11ShaderResourceView;
-struct ID3D11RenderTargetView;
-
 class WmfMediaHAPDecoder 
 	: public WmfMediaDecoder
 {
-
 public:
 
 	WmfMediaHAPDecoder();
@@ -28,6 +18,7 @@ public:
 	virtual HRESULT STDMETHODCALLTYPE GetOutputAvailableType(DWORD dwOutputStreamID, DWORD dwTypeIndex, IMFMediaType** ppType) override;
 	virtual HRESULT STDMETHODCALLTYPE ProcessMessage(MFT_MESSAGE_TYPE eMessage, ULONG_PTR ulParam) override;
 	virtual HRESULT STDMETHODCALLTYPE ProcessOutput(DWORD dwFlags, DWORD cOutputBufferCount, MFT_OUTPUT_DATA_BUFFER* pOutputSamples, DWORD* pdwStatus) override;
+	virtual HRESULT STDMETHODCALLTYPE GetAttributes(IMFAttributes** pAttributes) override;
 
 	virtual bool IsExternalBufferSupported() const override;
 
@@ -43,51 +34,8 @@ private:
 	virtual HRESULT OnCheckInputType(IMFMediaType *InMediaType) override;
 	virtual HRESULT OnSetInputType(IMFMediaType *InMediaType) override;
 
-	bool InitPipeline(void);
-	void CreateSRV(ID3D11Texture2D* InTexture, D3D11_TEXTURE2D_DESC& InTextureDesc, ID3D11ShaderResourceView** InSRV);
-	void BindSRV(ID3D11ShaderResourceView* InSRV);
-	void BindSRV(ID3D11ShaderResourceView* InSRVA, ID3D11ShaderResourceView* InSRVB);
-
-	struct FWorkTextures
-	{
-		TComPtr<ID3D11Texture2D> InputTextureColor;
-		TComPtr<ID3D11Texture2D> InputTextureAlpha;
-		TComPtr<ID3D11ShaderResourceView> SRV_InputColor;
-		TComPtr<ID3D11ShaderResourceView> SRV_InputAlpha;
-	};
-
-	struct FOutputTextures
-	{
-		TComPtr<ID3D11Texture2D> Texture;
-		TComPtr<ID3D11RenderTargetView> RTV_A;
-		TComPtr<ID3D11RenderTargetView> RTV_B;
-	};
-
-	struct VertexDescription
-	{
-		float X, Y, Z;
-		float U, V;
-	};
-
-	struct ConstantBuffer
-	{
-		int Mode;
-		int Index;
-		int WriteY;
-		int Padding;
-	};
-
-
 	GUID InputSubType;
-	TComPtr<ID3D11InputLayout> InputLayout;
-	TComPtr<ID3D11Buffer> VertexBuffer;
-	TComPtr<ID3D11Buffer> PixelBuffer;
-
-	ConstantBuffer FrameData;
-
-	TUniquePtr<FOutputTextures> OutputTexture;
-	TUniquePtr<FWorkTextures> WorkTextures;
-
+	TComPtr<IMFMediaBuffer> MediaBuffer;
 };
 
 #endif // WMFMEDIA_SUPPORTED_PLATFORM
diff --git a/Engine/Plugins/Media/WmfMedia/Source/WmfMedia/Private/Player/WmfMediaPlayer.cpp b/Engine/Plugins/Media/WmfMedia/Source/WmfMedia/Private/Player/WmfMediaPlayer.cpp
--- a/Engine/Plugins/Media/WmfMedia/Source/WmfMedia/Private/Player/WmfMediaPlayer.cpp
+++ b/Engine/Plugins/Media/WmfMedia/Source/WmfMedia/Private/Player/WmfMediaPlayer.cpp
@@ -194,8 +194,19 @@ void FWmfMediaPlayer::Tick()
 			Session->Initialize(Settings->LowLatency);
 			Tracks->ReInitialize();
 		}
-	
-		if (!Tracks->IsInitialized() || !Session->SetTopology(Tracks->CreateTopology(), Tracks->GetDuration()))
+
+		bool bOK = false;
+		if (Tracks->IsInitialized())
+		{
+			auto Topo = Tracks->CreateTopology();
+
+			if (Session->SetTopology(Topo, Tracks->GetDuration()))
+			{
+				bOK = true;
+			}
+		}
+
+		if (!bOK)
 		{
 			Session->Shutdown();
 			EventSink.ReceiveMediaEvent(EMediaEvent::MediaOpenFailed);
diff --git a/Engine/Plugins/Media/WmfMedia/Source/WmfMedia/Private/Player/WmfMediaTracks.cpp b/Engine/Plugins/Media/WmfMedia/Source/WmfMedia/Private/Player/WmfMediaTracks.cpp
--- a/Engine/Plugins/Media/WmfMedia/Source/WmfMedia/Private/Player/WmfMediaTracks.cpp
+++ b/Engine/Plugins/Media/WmfMedia/Source/WmfMedia/Private/Player/WmfMediaTracks.cpp
@@ -16,7 +16,6 @@
 
 #if WITH_ENGINE
 	#include "Engine/Engine.h"
-	#include "HardwareInfo.h"
 #endif
 
 #include "Player/WmfMediaTextureSample.h"
@@ -45,8 +44,7 @@
  *****************************************************************************/
 
 FWmfMediaTracks::FWmfMediaTracks()
-	: bHardwareAccelerated(false)
-	, AudioSamplePool(new FWmfMediaAudioSamplePool)
+	: AudioSamplePool(new FWmfMediaAudioSamplePool)
 	, MediaSourceChanged(false)
 	, SelectedAudioTrack(INDEX_NONE)
 	, SelectedCaptionTrack(INDEX_NONE)
@@ -183,22 +181,13 @@ TComPtr<IMFTopology> FWmfMediaTracks::CreateTopology()
 
 	UE_LOG(LogWmfMedia, Verbose, TEXT("Tracks %p: Created playback topology %p (media source %p)"), this, Topology.Get(), MediaSource.Get());
 
-	bHardwareAccelerated = false;
-
 	const UWmfMediaSettings* WmfMediaSettings = GetDefault<UWmfMediaSettings>();
 	if (WmfMediaSettings->HardwareAcceleratedVideoDecoding ||
 		WmfMediaSettings->bAreHardwareAcceleratedCodecRegistered)
 	{
-#if WITH_ENGINE
-		if (FHardwareInfo::GetHardwareInfo(NAME_RHI) == "D3D11")
-#else
-		if (1)
-#endif
-		{
-			WmfMediaTopologyLoader MediaTopologyLoader;
-			bHardwareAccelerated = MediaTopologyLoader.EnableHardwareAcceleration(Topology);
-			bHardwareAccelerated = bHardwareAccelerated || bVideoTrackRequestedHardwareAcceleration;
-		}
+		WmfMediaTopologyLoader MediaTopologyLoader;
+		bool bHardwareAccelerated = MediaTopologyLoader.EnableHardwareAcceleration(Topology) || bVideoTrackRequestedHardwareAcceleration;
+
 		UE_LOG(LogWmfMedia, Verbose, TEXT("Tracks %p: Video (media source %p) will be decoded on %s"), this, MediaSource.Get(), bHardwareAccelerated ? TEXT("GPU") : TEXT("CPU"));
 		Info += FString::Printf(TEXT("Video decoded on %s\n"), bHardwareAccelerated ? TEXT("GPU") : TEXT("CPU"));
 	}
@@ -1279,8 +1268,17 @@ bool FWmfMediaTracks::AddTrackToTopology(const FTrack& Track, IMFTopology& Topol
 	TComPtr<FWmfMediaStreamSink> MediaStreamSink;
 
 	const UWmfMediaSettings* WmfMediaSettings = GetDefault<UWmfMediaSettings>();
-	if ((GEngine != nullptr) && 
-		bHardwareAccelerated &&
+
+	if (VideoSamplePool)
+	{
+		delete VideoSamplePool;
+		VideoSamplePool = nullptr;
+	}
+
+#if WITH_ENGINE
+	if ((GEngine != nullptr) &&
+		!Track.bNoDirectXOutput &&
+		(WmfMediaSettings->HardwareAcceleratedVideoDecoding || WmfMediaSettings->bAreHardwareAcceleratedCodecRegistered) &&
 		MajorType == MFMediaType_Video &&
 		FPlatformMisc::VerifyWindowsVersion(6, 2) && // Windows 8
 		FWmfMediaStreamSink::Create(MFMediaType_Video, MediaStreamSink))
@@ -1301,12 +1299,14 @@ bool FWmfMediaTracks::AddTrackToTopology(const FTrack& Track, IMFTopology& Topol
 		bVideoTrackRequestedHardwareAcceleration = true;
 	}
 	else
+#endif
 	{
-		if (VideoSamplePool)
+		VideoSamplePool = new FWmfMediaTextureSamplePool();
+		if (!VideoSamplePool)
 		{
-			delete VideoSamplePool;
+			UE_LOG(LogWmfMedia, Verbose, TEXT("Tracks %p: Failed to configure output node for stream (no memory) %i"), this, Track.StreamIndex);
+			return false;
 		}
-		VideoSamplePool = new FWmfMediaTextureSamplePool();
 
 		if (FAILED(::MFCreateTopologyNode(MF_TOPOLOGY_OUTPUT_NODE, &OutputNode)) ||
 			FAILED(OutputNode->SetObject(OutputActivator)) ||
@@ -1581,7 +1581,8 @@ bool FWmfMediaTracks::AddStreamToTracks(uint32 StreamIndex, bool IsVideoDevice,
 		OutInfo += FString::Printf(TEXT("\t\tCodec: %s\n"), *TypeName);
 
 		// create output type
-		TComPtr<IMFMediaType> OutputType = WmfMedia::CreateOutputType(*MediaType, AllowNonStandardCodecs, IsVideoDevice);
+		bool bNoDirectXOutput;
+		TComPtr<IMFMediaType> OutputType = WmfMedia::CreateOutputType(*MediaType, bNoDirectXOutput, AllowNonStandardCodecs, IsVideoDevice);
 
 		if (!OutputType.IsValid())
 		{
@@ -1590,6 +1591,8 @@ bool FWmfMediaTracks::AddStreamToTracks(uint32 StreamIndex, bool IsVideoDevice,
 			continue;
 		}
 
+		Track->bNoDirectXOutput = bNoDirectXOutput;
+
 		// add format details
 		int32 FormatIndex = INDEX_NONE;
 
diff --git a/Engine/Plugins/Media/WmfMedia/Source/WmfMedia/Private/Player/WmfMediaTracks.h b/Engine/Plugins/Media/WmfMedia/Source/WmfMedia/Private/Player/WmfMediaTracks.h
--- a/Engine/Plugins/Media/WmfMedia/Source/WmfMedia/Private/Player/WmfMediaTracks.h
+++ b/Engine/Plugins/Media/WmfMedia/Source/WmfMedia/Private/Player/WmfMediaTracks.h
@@ -85,6 +85,7 @@ class FWmfMediaTracks
 		bool Protected;
 		int32 SelectedFormat;
 		DWORD StreamIndex;
+		bool bNoDirectXOutput;
 	};
 
 public:
@@ -301,9 +302,6 @@ private:
 
 private:
 
-	/** HW accelerated or not. */
-	bool bHardwareAccelerated;
-
 	/** Audio sample object pool. */
 	FWmfMediaAudioSamplePool* AudioSamplePool;
 
diff --git a/Engine/Plugins/Media/WmfMedia/Source/WmfMedia/Private/Wmf/WmfMediaHardwareVideoDecodingRendering.cpp b/Engine/Plugins/Media/WmfMedia/Source/WmfMedia/Private/Wmf/WmfMediaHardwareVideoDecodingRendering.cpp
--- a/Engine/Plugins/Media/WmfMedia/Source/WmfMedia/Private/Wmf/WmfMediaHardwareVideoDecodingRendering.cpp
+++ b/Engine/Plugins/Media/WmfMedia/Source/WmfMedia/Private/Wmf/WmfMediaHardwareVideoDecodingRendering.cpp
@@ -134,6 +134,8 @@ bool FWmfMediaHardwareVideoDecodingParameters::ConvertTextureFormat_RenderThread
 		SCOPED_DRAW_EVENT(RHICmdList, FWmfMediaHardwareVideoDecodingParameters_Convert);
 		SCOPED_GPU_STAT(RHICmdList, MediaTextureConversion);
 
+		RHICmdList.Transition(FRHITransitionInfo(InDstTexture, ERHIAccess::SRVMask, ERHIAccess::RTV));
+
 		FRHIRenderPassInfo RPInfo(InDstTexture, ERenderTargetActions::DontLoad_Store);
 		RHICmdList.BeginRenderPass(RPInfo, TEXT("ConvertTextureFormat"));
 
@@ -270,6 +272,7 @@ bool FWmfMediaHardwareVideoDecodingParameters::ConvertTextureFormat_RenderThread
 
 		RHICmdList.DrawPrimitive(0, 2, 1);
 		RHICmdList.EndRenderPass();
+		RHICmdList.Transition(FRHITransitionInfo(InDstTexture, ERHIAccess::RTV, ERHIAccess::SRVMask));
 	}
 
 	return true;
diff --git a/Engine/Plugins/Media/WmfMedia/Source/WmfMedia/Private/Wmf/WmfMediaSink.cpp b/Engine/Plugins/Media/WmfMedia/Source/WmfMedia/Private/Wmf/WmfMediaSink.cpp
--- a/Engine/Plugins/Media/WmfMedia/Source/WmfMedia/Private/Wmf/WmfMediaSink.cpp
+++ b/Engine/Plugins/Media/WmfMedia/Source/WmfMedia/Private/Wmf/WmfMediaSink.cpp
@@ -50,21 +50,6 @@ bool FWmfMediaSink::Initialize(TComPtr<FWmfMediaStreamSink> InStreamSink)
 		return false;
 	}
 
-	// Now only support for DX11 or DX12
-	ERHIInterfaceType RHIType = RHIGetInterfaceType();
-
-#if WMFMEDIA_PLAYER_VERSION >= 2
-	if (RHIType != ERHIInterfaceType::D3D11 && RHIType != ERHIInterfaceType::D3D12)
-#else
-	// Now only support for DX11
-	if (RHIType != ERHIInterfaceType::D3D11)
-#endif // WMFMEDIA_PLAYER_VERSION >= 2
-	{
-		return false;
-	}
-
-	DWORD StreamId = 0;
-
 	if (!InStreamSink->Initialize(*this))
 	{
 		return false;
@@ -510,44 +495,25 @@ bool FWmfMediaSink::CreateDXGIManagerAndDevice()
 
 	if (nullptr == DXGIManager)
 	{
-		Result = MFCreateDXGIDeviceManager(&resetToken, &DXGIManager);
-		if (FAILED(Result))
-		{
-			UE_LOG(LogWmfMedia, Error, TEXT("Unable to Create DXGI Device Manager: %s"), *WmfMedia::ResultToString(Result));
-			return false;
-		}
-
 		if (GDynamicRHI == nullptr)
 		{
 			UE_LOG(LogWmfMedia, Error, TEXT("Invalid Dynamic RHI"));
 			return false;
 		}
 
-		// Is this D3D12?
 		TRefCountPtr<IDXGIAdapter> DXGIAdapter;
 		ERHIInterfaceType RHIType = RHIGetInterfaceType();
-		if (RHIType == ERHIInterfaceType::D3D12)
+
+		// Is this D3D11?
+		if(RHIType == ERHIInterfaceType::D3D11)
 		{
-			TRefCountPtr<IDXGIFactory4> DXGIFactory;
-			Result = CreateDXGIFactory(__uuidof(IDXGIFactory4), (void**)DXGIFactory.GetInitReference());
-			if (FAILED(Result))
-			{
-				UE_LOG(LogWmfMedia, Error, TEXT("CreateDXGIFactory failed."));
-				return false;
-			}
-			
-			ID3D12Device* UE4DxDevice = static_cast<ID3D12Device*>(GDynamicRHI->RHIGetNativeDevice());
-			LUID Luid = UE4DxDevice->GetAdapterLuid();
-			Result = DXGIFactory->EnumAdapterByLuid(Luid, __uuidof(IDXGIAdapter), (void**)DXGIAdapter.GetInitReference());
+			Result = MFCreateDXGIDeviceManager(&resetToken, &DXGIManager);
 			if (FAILED(Result))
 			{
-				UE_LOG(LogWmfMedia, Error, TEXT("EnumAdapterByLuid failed."));
+				UE_LOG(LogWmfMedia, Error, TEXT("Unable to Create DXGI Device Manager: %s"), *WmfMedia::ResultToString(Result));
 				return false;
 			}
-		}
-		// Is this D3D11?
-		else if(RHIType == ERHIInterfaceType::D3D11)
-		{
+
 			ID3D11Device* PreExistingD3D11Device = static_cast<ID3D11Device*>(GDynamicRHI->RHIGetNativeDevice());
 
 			TComPtr<IDXGIDevice> DXGIDevice;
@@ -557,8 +523,10 @@ bool FWmfMediaSink::CreateDXGIManagerAndDevice()
 		}
 		else
 		{
-			UE_LOG(LogWmfMedia, Error, TEXT("Dynamic RHI is not D3D11 nor D3D12"));
-			return false;
+			DXGIManager.Reset();
+			D3D11Device.Reset();
+			UE_LOG(LogWmfMedia, Verbose, TEXT("Dynamic RHI is not D3D11, skipping DX11 device creation for AV decoding"));
+			return true;
 		}
 
 		// Create device from same adapter as already existing device
diff --git a/Engine/Plugins/Media/WmfMedia/Source/WmfMedia/Private/Wmf/WmfMediaStreamSink.cpp b/Engine/Plugins/Media/WmfMedia/Source/WmfMedia/Private/Wmf/WmfMediaStreamSink.cpp
--- a/Engine/Plugins/Media/WmfMedia/Source/WmfMedia/Private/Wmf/WmfMediaStreamSink.cpp
+++ b/Engine/Plugins/Media/WmfMedia/Source/WmfMedia/Private/Wmf/WmfMediaStreamSink.cpp
@@ -845,10 +845,10 @@ void FWmfMediaStreamSink::CopyTextureAndEnqueueSample(IMFSample* pSample)
 	}
 
 	DWORD cBuffers = 0;
-	IMFMediaBuffer* pBuffer = nullptr;
-	IMFDXGIBuffer* pDXGIBuffer = nullptr;
+	TComPtr<IMFMediaBuffer> pBuffer;
+	TComPtr<IMFDXGIBuffer> pDXGIBuffer;
 	UINT dwViewIndex = 0;
-	ID3D11Texture2D* pTexture2D = nullptr;
+	TComPtr<ID3D11Texture2D> pTexture2D;
 
 	HRESULT Result = pSample->GetBufferCount(&cBuffers);
 	if (FAILED(Result))
@@ -863,23 +863,6 @@ void FWmfMediaStreamSink::CopyTextureAndEnqueueSample(IMFSample* pSample)
 			return;
 		}
 	}
-	Result = pBuffer->QueryInterface(__uuidof(IMFDXGIBuffer), (LPVOID*)&pDXGIBuffer);
-	if (FAILED(Result))
-	{
-		return;
-	}
-
-	Result = pDXGIBuffer->GetResource(__uuidof(ID3D11Texture2D), (LPVOID*)&pTexture2D);
-	if (FAILED(Result))
-	{
-		return;
-	}
-
-	Result = pDXGIBuffer->GetSubresourceIndex(&dwViewIndex);
-	if (FAILED(Result))
-	{
-		return;
-	}
 
 	UINT32 DimX;
 	UINT32 DimY;
@@ -896,6 +879,12 @@ void FWmfMediaStreamSink::CopyTextureAndEnqueueSample(IMFSample* pSample)
 	{
 		const TSharedRef<FWmfMediaHardwareVideoDecodingTextureSample, ESPMode::ThreadSafe> TextureSample = VideoSamplePool->AcquireShared();
 
+		GUID WrappedFormatGuid;
+		if (SUCCEEDED(CurrentMediaType->GetGUID(UE_WMF_PrivateFormatGUID, &WrappedFormatGuid)))
+		{
+			Guid = WrappedFormatGuid;
+		}
+
 		EPixelFormat PixelFormat = PF_Unknown;
 		EPixelFormat AlphaPixelFormat = PF_Unknown;
 		EMediaTextureSampleFormat MediaTextureSampleFormat = EMediaTextureSampleFormat::Undefined;
@@ -988,6 +977,24 @@ void FWmfMediaStreamSink::CopyTextureAndEnqueueSample(IMFSample* pSample)
 		}
 		else
 		{
+			Result = pBuffer->QueryInterface(__uuidof(IMFDXGIBuffer), (LPVOID*)&pDXGIBuffer);
+			if (FAILED(Result))
+			{
+				return;
+			}
+
+			Result = pDXGIBuffer->GetResource(__uuidof(ID3D11Texture2D), (LPVOID*)&pTexture2D);
+			if (FAILED(Result))
+			{
+				return;
+			}
+
+			Result = pDXGIBuffer->GetSubresourceIndex(&dwViewIndex);
+			if (FAILED(Result))
+			{
+				return;
+			}
+
 			check(TextureSample->GetMediaTextureSampleConverter() != nullptr);
 			ID3D11Texture2D* SharedTexture = TextureSample->InitializeSourceTexture(
 				Owner->GetDevice(),
@@ -998,7 +1005,9 @@ void FWmfMediaStreamSink::CopyTextureAndEnqueueSample(IMFSample* pSample)
 				MediaTextureSampleFormat);
 
 			if (!SharedTexture)
+			{
 				return;
+			}
 
 			D3D11_BOX SrcBox;
 			SrcBox.left = 0;
@@ -1037,19 +1046,6 @@ void FWmfMediaStreamSink::CopyTextureAndEnqueueSample(IMFSample* pSample)
 			bShowSubTypeErrorMessage = false;
 		}
 	}
-
-	if (pTexture2D)
-	{
-		pTexture2D->Release();
-	}
-	if (pDXGIBuffer)
-	{
-		pDXGIBuffer->Release();
-	}
-	if (pBuffer)
-	{
-		pBuffer->Release();
-	}
 }
 
 
diff --git a/Engine/Plugins/Media/WmfMedia/Source/WmfMedia/Private/Wmf/WmfMediaUtils.cpp b/Engine/Plugins/Media/WmfMedia/Source/WmfMedia/Private/Wmf/WmfMediaUtils.cpp
--- a/Engine/Plugins/Media/WmfMedia/Source/WmfMedia/Private/Wmf/WmfMediaUtils.cpp
+++ b/Engine/Plugins/Media/WmfMedia/Source/WmfMedia/Private/Wmf/WmfMediaUtils.cpp
@@ -8,6 +8,9 @@
 #include "HAL/PlatformMisc.h"
 #include "HAL/PlatformProcess.h"
 
+#include "RHI.h"
+#include "DynamicRHI.h"
+
 #include "IWmfMediaModule.h"
 #include "Math/NumericLimits.h"
 #include "Misc/FileHelper.h"
@@ -355,7 +358,7 @@ namespace WmfMedia
 		return FMemory::Memcmp(&lhs.Data2, &rhs.Data2, 12) == 0; //-V512
 	}
 
-	TComPtr<IMFMediaType> CreateOutputType(IMFMediaType& InputType, bool AllowNonStandardCodecs, bool IsVideoDevice)
+	TComPtr<IMFMediaType> CreateOutputType(IMFMediaType& InputType, bool& bNoDirectXOutput, bool AllowNonStandardCodecs, bool IsVideoDevice)
 	{
 		GUID MajorType;
 		{
@@ -398,6 +401,8 @@ namespace WmfMedia
 			}
 		}
 
+		bNoDirectXOutput = false;
+
 		if (MajorType == MFMediaType_Audio)
 		{
 			// filter unsupported audio formats
@@ -504,6 +509,9 @@ namespace WmfMedia
 					UE_LOG(LogWmfMedia, Warning, TEXT("H264 video type requires Windows 8 or newer (your version is %s)"), *FPlatformMisc::GetOSVersion());
 					return NULL;
 				}
+
+				// Is this D3D11? Only then we can output directly into a texture
+				bNoDirectXOutput = (RHIGetInterfaceType() != ERHIInterfaceType::D3D11);
 			}
 
 			if ((SubType == MFVideoFormat_HEVC) || (SubType == MFVideoFormat_HEVC_ES))
@@ -518,6 +526,9 @@ namespace WmfMedia
 
 					UE_LOG(LogWmfMedia, Warning, TEXT("HEVC video type requires Windows 10 or newer (your version is %s), and game must be manifested for Windows 10"), *FPlatformMisc::GetOSVersion());
 				}
+
+				// Is this D3D11? Only then we can output directly into a texture
+				bNoDirectXOutput = (RHIGetInterfaceType() != ERHIInterfaceType::D3D11);
 			}
 
 			// configure video output
diff --git a/Engine/Plugins/Media/WmfMedia/Source/WmfMedia/Private/Wmf/WmfMediaUtils.h b/Engine/Plugins/Media/WmfMedia/Source/WmfMedia/Private/Wmf/WmfMediaUtils.h
--- a/Engine/Plugins/Media/WmfMedia/Source/WmfMedia/Private/Wmf/WmfMediaUtils.h
+++ b/Engine/Plugins/Media/WmfMedia/Source/WmfMedia/Private/Wmf/WmfMediaUtils.h
@@ -49,7 +49,7 @@ namespace WmfMedia
 	 * @param Whether the media source is a video capture device (there are some limitations with those).
 	 * @return The output type, or NULL on error.
 	 */
-	TComPtr<IMFMediaType> CreateOutputType(IMFMediaType& InputType, bool AllowNonStandardCodecs, bool IsVideoDevice);
+	TComPtr<IMFMediaType> CreateOutputType(IMFMediaType& InputType, bool& bNoDirectXOutput, bool AllowNonStandardCodecs, bool IsVideoDevice);
 
 	/**
 	 * Dump information about the given media type to the output log.
diff --git a/Engine/Plugins/Media/WmfMedia/Source/WmfMedia/Public/WmfMediaCommon.h b/Engine/Plugins/Media/WmfMedia/Source/WmfMedia/Public/WmfMediaCommon.h
--- a/Engine/Plugins/Media/WmfMedia/Source/WmfMedia/Public/WmfMediaCommon.h
+++ b/Engine/Plugins/Media/WmfMedia/Source/WmfMedia/Public/WmfMediaCommon.h
@@ -44,6 +44,9 @@
 		const GUID MF_LOW_LATENCY = { 0x9c27891a, 0xed7a, 0x40e1, { 0x88, 0xe8, 0xb2, 0x27, 0x27, 0xa0, 0x24, 0xee } };
 	#endif
 
+	// A private GUID used to add a wrapped format GUID to an IMFMediaType.
+	const GUID UE_WMF_PrivateFormatGUID = { 0x630d3c34, 0xefac, 0x42f0, {0xa9, 0x6f, 0x6e, 0xa2, 0x6f, 0xd4, 0x16, 0x6}};
+
 	#include "Microsoft/COMPointer.h"
 	
 	#include "Windows/PostWindowsApi.h"
